[OperatorStack] Moved plugin

diff --git a/Engine/Plugins/Experimental/OperatorStack/OperatorStack.uplugin b/Engine/Plugins/Experimental/OperatorStack/OperatorStack.uplugin
new file mode 100644
--- /dev/null
+++ b/Engine/Plugins/Experimental/OperatorStack/OperatorStack.uplugin
@@ -0,0 +1,33 @@
+﻿{
+	"FileVersion": 3,
+	"Version": 1,
+	"VersionName": "1.0",
+	"FriendlyName": "Operator Stack",
+	"Description": "Allows you to display linked objects in a container with an intuitive UI",
+	"Category": "UI",
+	"CreatedBy": "Epic Games, Inc.",
+	"CreatedByURL": "",
+	"DocsURL": "",
+	"MarketplaceURL": "",
+	"CanContainContent": true,
+	"IsBetaVersion": false,
+	"IsExperimentalVersion": false,
+	"Installed": false,
+	"Modules": [
+		{
+			"Name": "OperatorStackEditor",
+			"Type": "Editor",
+			"LoadingPhase": "Default"
+		}
+	],
+	"Plugins": [
+		{
+			"Name": "StructUtils",
+			"Enabled": true
+		},
+		{
+			"Name": "CustomDetailsView",
+			"Enabled": true
+		}
+	]
+}
\ No newline at end of file
diff --git a/Engine/Plugins/Experimental/OperatorStack/Source/OperatorStack/OperatorStackEditor.Build.cs b/Engine/Plugins/Experimental/OperatorStack/Source/OperatorStack/OperatorStackEditor.Build.cs
new file mode 100644
--- /dev/null
+++ b/Engine/Plugins/Experimental/OperatorStack/Source/OperatorStack/OperatorStackEditor.Build.cs
@@ -0,0 +1,39 @@
+﻿// Copyright Epic Games, Inc. All Rights Reserved.
+
+using UnrealBuildTool;
+
+public class OperatorStackEditor : ModuleRules
+{
+	public OperatorStackEditor(ReadOnlyTargetRules Target) : base(Target)
+	{
+		PCHUsage = ModuleRules.PCHUsageMode.UseExplicitOrSharedPCHs;
+
+		PublicDependencyModuleNames.AddRange(
+			new string[]
+			{
+				"Core",
+				"CoreUObject"
+			}
+		);
+
+		PrivateDependencyModuleNames.AddRange(
+			new string[]
+			{
+				"ApplicationCore",
+				"CoreUObject",
+				"CustomDetailsView",
+				"EditorSubsystem",
+				"Engine",
+				"InputCore",
+				"Projects",
+				"PropertyEditor",
+				"Slate",
+				"SlateCore",
+				"ToolMenus",
+				"ToolWidgets",
+				"UnrealEd",
+				"WorkspaceMenuStructure"
+			}
+		);
+	}
+}
diff --git a/Engine/Plugins/Experimental/OperatorStack/Source/OperatorStack/Private/Builders/OperatorStackEditorBodyBuilder.cpp b/Engine/Plugins/Experimental/OperatorStack/Source/OperatorStack/Private/Builders/OperatorStackEditorBodyBuilder.cpp
new file mode 100644
--- /dev/null
+++ b/Engine/Plugins/Experimental/OperatorStack/Source/OperatorStack/Private/Builders/OperatorStackEditorBodyBuilder.cpp
@@ -0,0 +1,52 @@
+// Copyright Epic Games, Inc. All Rights Reserved.
+
+#include "Builders/OperatorStackEditorBodyBuilder.h"
+#include "Items/CustomDetailsViewItemId.h"
+
+FOperatorStackEditorBodyBuilder& FOperatorStackEditorBodyBuilder::SetCustomWidget(TSharedPtr<SWidget> InWidget)
+{
+	CustomWidget = InWidget;
+	return *this;
+}
+
+FOperatorStackEditorBodyBuilder& FOperatorStackEditorBodyBuilder::SetShowDetailsView(bool bInShowDetailsView)
+{
+	bShowDetailsView = bInShowDetailsView;
+	return *this;
+}
+
+FOperatorStackEditorBodyBuilder& FOperatorStackEditorBodyBuilder::SetDetailsViewItem(TSharedPtr<FOperatorStackEditorItem> InItem)
+{
+	DetailsViewItem = InItem;
+	return *this;
+}
+
+FOperatorStackEditorBodyBuilder& FOperatorStackEditorBodyBuilder::DisallowProperty(FProperty* InProperty)
+{
+	DisallowedDetailsViewItems.Add(MakeShared<FCustomDetailsViewItemId>(FCustomDetailsViewItemId::MakePropertyId(InProperty)));
+	return *this;
+}
+
+FOperatorStackEditorBodyBuilder& FOperatorStackEditorBodyBuilder::DisallowCategory(const FName& InCategory)
+{
+	DisallowedDetailsViewItems.Add(MakeShared<FCustomDetailsViewItemId>(FCustomDetailsViewItemId::MakeCategoryId(InCategory)));
+	return *this;
+}
+
+FOperatorStackEditorBodyBuilder& FOperatorStackEditorBodyBuilder::AllowProperty(FProperty* InProperty)
+{
+	AllowedDetailsViewItems.Add(MakeShared<FCustomDetailsViewItemId>(FCustomDetailsViewItemId::MakePropertyId(InProperty)));
+	return *this;
+}
+
+FOperatorStackEditorBodyBuilder& FOperatorStackEditorBodyBuilder::AllowCategory(const FName& InCategory)
+{
+	AllowedDetailsViewItems.Add(MakeShared<FCustomDetailsViewItemId>(FCustomDetailsViewItemId::MakeCategoryId(InCategory)));
+	return *this;
+}
+
+FOperatorStackEditorBodyBuilder& FOperatorStackEditorBodyBuilder::SetEmptyBodyText(const FText& InText)
+{
+	EmptyBodyText = InText;
+	return *this;
+}
diff --git a/Engine/Plugins/Experimental/OperatorStack/Source/OperatorStack/Private/Builders/OperatorStackEditorHeaderBuilder.cpp b/Engine/Plugins/Experimental/OperatorStack/Source/OperatorStack/Private/Builders/OperatorStackEditorHeaderBuilder.cpp
new file mode 100644
--- /dev/null
+++ b/Engine/Plugins/Experimental/OperatorStack/Source/OperatorStack/Private/Builders/OperatorStackEditorHeaderBuilder.cpp
@@ -0,0 +1,104 @@
+// Copyright Epic Games, Inc. All Rights Reserved.
+
+#include "Builders/OperatorStackEditorHeaderBuilder.h"
+
+#include "Items/CustomDetailsViewItemId.h"
+
+FOperatorStackEditorHeaderBuilder& FOperatorStackEditorHeaderBuilder::SetIcon(const FSlateBrush* InIcon)
+{
+	Icon = InIcon;
+	return *this;
+}
+
+FOperatorStackEditorHeaderBuilder& FOperatorStackEditorHeaderBuilder::SetLabel(const FText& InLabel)
+{
+	Label = InLabel;
+	return *this;
+}
+
+FOperatorStackEditorHeaderBuilder& FOperatorStackEditorHeaderBuilder::SetBorderColor(const FLinearColor& InColor)
+{
+	BorderColor = InColor;
+	return *this;
+}
+
+FOperatorStackEditorHeaderBuilder& FOperatorStackEditorHeaderBuilder::SetStartsExpanded(bool bInExpanded)
+{
+	bStartsExpanded = bInExpanded;
+	return *this;
+}
+
+FOperatorStackEditorHeaderBuilder& FOperatorStackEditorHeaderBuilder::SetDraggable(bool bInDraggable)
+{
+	bDraggable = bInDraggable;
+	return *this;
+}
+
+FOperatorStackEditorHeaderBuilder& FOperatorStackEditorHeaderBuilder::SetCommandList(TSharedPtr<FUICommandList> InCommands)
+{
+	CommandList = InCommands;
+	return *this;
+}
+
+FOperatorStackEditorHeaderBuilder& FOperatorStackEditorHeaderBuilder::SetContextMenu(const FName& InToolMenuName)
+{
+	ContextMenuName = InToolMenuName;
+	return *this;
+}
+
+FOperatorStackEditorHeaderBuilder& FOperatorStackEditorHeaderBuilder::SetCustomWidget(TSharedPtr<SWidget> InWidget)
+{
+	CustomWidget = InWidget;
+	return *this;
+}
+
+FOperatorStackEditorHeaderBuilder& FOperatorStackEditorHeaderBuilder::SetMessageBox(const TAttribute<EOperatorStackEditorMessageType>& InType, const TAttribute<FText>& InText)
+{
+	MessageBoxType = InType;
+	MessageBoxText = InText;
+	return *this;
+}
+
+FOperatorStackEditorHeaderBuilder& FOperatorStackEditorHeaderBuilder::SetProperty(FProperty* InProperty)
+{
+	Property = MakeShared<FCustomDetailsViewItemId>(FCustomDetailsViewItemId::MakePropertyId(InProperty));
+	return *this;
+}
+
+FOperatorStackEditorHeaderBuilder& FOperatorStackEditorHeaderBuilder::SetToolbarMenu(const FName& InActionMenuName)
+{
+	ToolbarMenuName = InActionMenuName;
+	return *this;
+}
+
+FOperatorStackEditorHeaderBuilder& FOperatorStackEditorHeaderBuilder::SetToolMenu(const FName& InToolMenuName, const FText& InToolMenuLabel, const FSlateBrush* InToolMenuIcon)
+{
+	ToolMenuName = InToolMenuName;
+	ToolMenuLabel = InToolMenuLabel;
+	ToolMenuIcon = InToolMenuIcon;
+	return *this;
+}
+
+FOperatorStackEditorHeaderBuilder& FOperatorStackEditorHeaderBuilder::SetSearchAllowed(bool bInSearchable)
+{
+	bSearchAllowed = bInSearchable;
+	return *this;
+}
+
+FOperatorStackEditorHeaderBuilder& FOperatorStackEditorHeaderBuilder::SetSearchKeywords(const TSet<FString>& InKeywords)
+{
+	SearchKeywords = InKeywords;
+	return *this;
+}
+
+FOperatorStackEditorHeaderBuilder& FOperatorStackEditorHeaderBuilder::SetSearchPinnedKeywords(const TSet<FString>& InKeywords)
+{
+	SearchPinnedKeywords = InKeywords;
+	return *this;
+}
+
+FOperatorStackEditorHeaderBuilder& FOperatorStackEditorHeaderBuilder::SetExpandable(bool bInExpandable)
+{
+	bExpandable = bInExpandable;
+	return *this;
+}
diff --git a/Engine/Plugins/Experimental/OperatorStack/Source/OperatorStack/Private/Customizations/OperatorStackEditorStackCustomization.cpp b/Engine/Plugins/Experimental/OperatorStack/Source/OperatorStack/Private/Customizations/OperatorStackEditorStackCustomization.cpp
new file mode 100644
--- /dev/null
+++ b/Engine/Plugins/Experimental/OperatorStack/Source/OperatorStack/Private/Customizations/OperatorStackEditorStackCustomization.cpp
@@ -0,0 +1,75 @@
+// Copyright Epic Games, Inc. All Rights Reserved.
+
+#include "Customizations/OperatorStackEditorStackCustomization.h"
+#include "Styling/SlateIconFinder.h"
+
+bool UOperatorStackEditorStackCustomization::RegisterCustomizationFor(const UStruct* InItemDefinition)
+{
+	bool bAlreadyInSet = false;
+	SupportedDefinitions.Add(InItemDefinition, &bAlreadyInSet);
+	return !bAlreadyInSet;
+}
+
+bool UOperatorStackEditorStackCustomization::RegisterCustomizationFor(const FFieldClass* InItemDefinition)
+{
+	bool bAlreadyInSet = false;
+	SupportedFieldClasses.Add(InItemDefinition, &bAlreadyInSet);
+	return !bAlreadyInSet;
+}
+
+bool UOperatorStackEditorStackCustomization::UnregisterCustomizationFor(const UStruct* InItemDefinition)
+{
+	return SupportedDefinitions.Remove(InItemDefinition) > 0;
+}
+
+bool UOperatorStackEditorStackCustomization::UnregisterCustomizationFor(const FFieldClass* InItemDefinition)
+{
+	return SupportedFieldClasses.Remove(InItemDefinition) > 0;
+}
+
+bool UOperatorStackEditorStackCustomization::IsCustomizationSupportedFor(const FOperatorStackEditorItemPtr& InItem) const
+{
+	if (!InItem.IsValid())
+	{
+		return false;
+	}
+	
+	if (const UStruct* ItemStruct = InItem->GetValueType().Get<UStruct>())
+	{
+		if (SupportedDefinitions.Contains(ItemStruct))
+		{
+			return true;
+		}
+	
+		for (const TObjectPtr<const UStruct>& SupportedDefinition : SupportedDefinitions)
+		{
+			if (ItemStruct->IsChildOf(SupportedDefinition.Get()))
+			{
+				return true;
+			}
+		}
+	}
+
+	if (const FFieldClass* ItemClass = InItem->GetValueType().Get<FFieldClass>())
+	{
+		if (SupportedFieldClasses.Contains(ItemClass))
+		{
+			return true;
+		}
+	
+		for (const FFieldClass* SupportedFieldClass : SupportedFieldClasses)
+		{
+			if (ItemClass->IsChildOf(SupportedFieldClass))
+			{
+				return true;
+			}
+		}
+	}
+	
+	return false;
+}
+
+const FSlateBrush* UOperatorStackEditorStackCustomization::GetIcon() const
+{
+	return FSlateIconFinder::FindIconBrushForClass(GetClass());
+}
diff --git a/Engine/Plugins/Experimental/OperatorStack/Source/OperatorStack/Private/Items/OperatorStackEditorTree.cpp b/Engine/Plugins/Experimental/OperatorStack/Source/OperatorStack/Private/Items/OperatorStackEditorTree.cpp
new file mode 100644
--- /dev/null
+++ b/Engine/Plugins/Experimental/OperatorStack/Source/OperatorStack/Private/Items/OperatorStackEditorTree.cpp
@@ -0,0 +1,214 @@
+// Copyright Epic Games, Inc. All Rights Reserved.
+
+#include "Items/OperatorStackEditorTree.h"
+
+#include "Customizations/OperatorStackEditorStackCustomization.h"
+
+FOperatorStackEditorTree::FOperatorStackEditorTreeNode::FOperatorStackEditorTreeNode(int32 InItemIndex, int32 InParentIndex)
+{
+	check(InItemIndex >= 0)
+
+	ItemIndex = InItemIndex;
+	ParentIndex = InParentIndex;
+}
+
+FOperatorStackEditorTree::FOperatorStackEditorTree(UOperatorStackEditorStackCustomization* InCustomization, FOperatorStackEditorContextPtr InContext)
+	: CustomizationWeak(InCustomization)
+{
+	check(InCustomization && InContext.IsValid());
+
+	ContextWeak = InContext;
+
+	// Build if context is valid
+	const TArray<FOperatorStackEditorItemPtr> ContextItems(InContext->GetItems());
+	BuildTree(ContextItems);
+}
+
+TArray<FOperatorStackEditorItemPtr> FOperatorStackEditorTree::GetRootItems() const
+{
+	TArray<FOperatorStackEditorItemPtr> RootItems;
+	RootItems.Reserve(RootNodes.Num());
+
+	for (const FOperatorStackEditorTreeNodePtr& Node : RootNodes)
+	{
+		check(Items.IsValidIndex(Node->ItemIndex))
+
+		RootItems.Add(Items[Node->ItemIndex]);
+	}
+
+	return RootItems;
+}
+
+TArray<FOperatorStackEditorItemPtr> FOperatorStackEditorTree::GetChildrenItems(FOperatorStackEditorItemPtr InItem) const
+{
+	TArray<FOperatorStackEditorItemPtr> Children;
+
+	const int32 Index = Items.Find(InItem);
+
+	if (Nodes.IsValidIndex(Index))
+	{
+		Children.Reserve(Nodes[Index]->ChildrenIndices.Num());
+
+		for (const int32 ChildIndex : Nodes[Index]->ChildrenIndices)
+		{
+			check(Items.IsValidIndex(ChildIndex))
+
+			Children.Add(Items[ChildIndex]);
+		}
+	}
+
+	return Children;
+}
+
+FOperatorStackEditorItemPtr FOperatorStackEditorTree::GetParentItem(FOperatorStackEditorItemPtr InItem) const
+{
+	const int32 Index = Items.Find(InItem);
+
+	if (Nodes.IsValidIndex(Index))
+	{
+		const int32 ParentIndex = Nodes[Index]->ParentIndex;
+
+		if (Items.IsValidIndex(ParentIndex))
+		{
+			return Items[Nodes[Index]->ParentIndex];
+		}
+	}
+
+	return nullptr;
+}
+
+TArray<FOperatorStackEditorItemPtr> FOperatorStackEditorTree::GetLeafItems() const
+{
+	TArray<FOperatorStackEditorItemPtr> LeafItems;
+
+	for (const FOperatorStackEditorTreeNodePtr& Node : Nodes)
+	{
+		if (Node->ChildrenIndices.IsEmpty())
+		{
+			check(Items.IsValidIndex(Node->ItemIndex))
+
+			LeafItems.Add(Items[Node->ItemIndex]);
+		}
+	}
+
+	return LeafItems;
+}
+
+TConstArrayView<FOperatorStackEditorItemPtr> FOperatorStackEditorTree::GetAllItems() const
+{
+	return Items;
+}
+
+const FOperatorStackEditorContext& FOperatorStackEditorTree::GetContext() const
+{
+	const FOperatorStackEditorContextPtr Context = ContextWeak.Pin();
+	check(Context.IsValid())
+	return *Context;
+}
+
+UOperatorStackEditorStackCustomization* FOperatorStackEditorTree::GetCustomization() const
+{
+	return CustomizationWeak.Get();
+}
+
+bool FOperatorStackEditorTree::Contains(FOperatorStackEditorItemPtr InItem) const
+{
+	if (!InItem.IsValid())
+	{
+		return false;
+	}
+
+	return Items.ContainsByPredicate([InItem](FOperatorStackEditorItemPtr InTreeItem)
+	{
+		if (InTreeItem.IsValid())
+		{
+			return *InTreeItem.Get() == *InItem.Get();
+		}
+
+		return false;
+	});
+}
+
+void FOperatorStackEditorTree::GetSupportedItems(const TArray<FOperatorStackEditorItemPtr>& InItems, TArray<FOperatorStackEditorItemPtr>& OutSupportedItems) const
+{
+	const UOperatorStackEditorStackCustomization* StackCustomization = CustomizationWeak.Get();
+
+	// Gather supported items for this customization
+	for (const FOperatorStackEditorItemPtr& Item : InItems)
+	{
+		if (!Item.IsValid() || !Item->HasValue())
+		{
+			continue;
+		}
+
+		if (StackCustomization->IsCustomizationSupportedFor(Item))
+		{
+			OutSupportedItems.Add(Item);
+			continue;
+		}
+
+		// Item is not supported, let's transform it to something supported
+		OutSupportedItems.Append(GetSupportedChildrenItems(Item));
+	}
+}
+
+TArray<FOperatorStackEditorItemPtr> FOperatorStackEditorTree::GetSupportedChildrenItems(const FOperatorStackEditorItemPtr& InParentItem) const
+{
+	TArray<FOperatorStackEditorItemPtr> ChildrenPtr;
+
+	if (!InParentItem.IsValid() || !InParentItem->HasValue())
+	{
+		return ChildrenPtr;
+	}
+
+	const UOperatorStackEditorStackCustomization* StackCustomization = CustomizationWeak.Get();
+	StackCustomization->TransformContextItem(InParentItem, ChildrenPtr);
+	ChildrenPtr.RemoveAll([](const FOperatorStackEditorItemPtr& InItem)
+	{
+		return !InItem.IsValid() || !InItem->HasValue();
+	});
+
+	return ChildrenPtr;
+}
+
+void FOperatorStackEditorTree::BuildTree(const TArray<FOperatorStackEditorItemPtr>& InSourceItems)
+{
+	TArray<FOperatorStackEditorItemPtr> SupportedItems;
+	GetSupportedItems(InSourceItems, SupportedItems);
+
+	Items.Empty(SupportedItems.Num());
+
+	BuildTreeInternal(SupportedItems, INDEX_NONE);
+}
+
+void FOperatorStackEditorTree::BuildTreeInternal(const TArray<FOperatorStackEditorItemPtr>& InSupportedItems, int32 InParentIndex)
+{
+	for (const FOperatorStackEditorItemPtr& SupportedItem : InSupportedItems)
+	{
+		// Do not add item again
+		if (!SupportedItem.IsValid()
+			|| !SupportedItem->HasValue()
+			|| Contains(SupportedItem))
+		{
+			continue;
+		}
+
+		const int32 ItemIndex = Items.Add(SupportedItem);
+
+		FOperatorStackEditorTreeNodePtr Node = MakeShared<FOperatorStackEditorTreeNode>(ItemIndex, InParentIndex);
+		Nodes.Add(Node);
+
+		if (InParentIndex == INDEX_NONE)
+		{
+			RootNodes.Add(Node);
+		}
+		else
+		{
+			check(Nodes.IsValidIndex(InParentIndex))
+
+			Nodes[InParentIndex]->ChildrenIndices.AddUnique(ItemIndex);
+		}
+
+		BuildTreeInternal(GetSupportedChildrenItems(SupportedItem), ItemIndex);
+	}
+}
diff --git a/Engine/Plugins/Experimental/OperatorStack/Source/OperatorStack/Private/OperatorStackEditorModule.cpp b/Engine/Plugins/Experimental/OperatorStack/Source/OperatorStack/Private/OperatorStackEditorModule.cpp
new file mode 100644
--- /dev/null
+++ b/Engine/Plugins/Experimental/OperatorStack/Source/OperatorStack/Private/OperatorStackEditorModule.cpp
@@ -0,0 +1,5 @@
+﻿// Copyright Epic Games, Inc. All Rights Reserved.
+
+#include "Modules/ModuleManager.h"
+	
+IMPLEMENT_MODULE(FDefaultModuleImpl, OperatorStackEditor)
\ No newline at end of file
diff --git a/Engine/Plugins/Experimental/OperatorStack/Source/OperatorStack/Private/Styles/OperatorStackEditorStyle.cpp b/Engine/Plugins/Experimental/OperatorStack/Source/OperatorStack/Private/Styles/OperatorStackEditorStyle.cpp
new file mode 100644
--- /dev/null
+++ b/Engine/Plugins/Experimental/OperatorStack/Source/OperatorStack/Private/Styles/OperatorStackEditorStyle.cpp
@@ -0,0 +1,30 @@
+// Copyright Epic Games, Inc. All Rights Reserved.
+
+#include "Styles/OperatorStackEditorStyle.h"
+#include "Misc/Paths.h"
+#include "Styling/StyleColors.h"
+#include "Styling/SlateTypes.h"
+#include "Styling/SlateStyleRegistry.h"
+
+FOperatorStackEditorStyle::FOperatorStackEditorStyle()
+	: FSlateStyleSet(TEXT("OperatorStackEditor"))
+{
+	const FLinearColor BackgroundColor = COLOR("#1A1A1AFF");
+	
+	Set("HoverColor", COLOR("#808080FF"));
+	Set("PropertiesBackgroundColor", COLOR("#1A1A1AFF"));
+	Set("BackgroundColor", BackgroundColor);
+	Set("ForegroundColor", COLOR("#242424FF"));
+	
+	const FTableViewStyle DefaultTreeViewStyle = FTableViewStyle()
+		.SetBackgroundBrush(FSlateColorBrush(BackgroundColor));
+	
+	Set("ListViewStyle", DefaultTreeViewStyle);
+	
+	FSlateStyleRegistry::RegisterSlateStyle(*this);
+}
+
+FOperatorStackEditorStyle::~FOperatorStackEditorStyle()
+{
+	FSlateStyleRegistry::UnRegisterSlateStyle(*this);
+}
\ No newline at end of file
diff --git a/Engine/Plugins/Experimental/OperatorStack/Source/OperatorStack/Private/Styles/OperatorStackEditorStyle.h b/Engine/Plugins/Experimental/OperatorStack/Source/OperatorStack/Private/Styles/OperatorStackEditorStyle.h
new file mode 100644
--- /dev/null
+++ b/Engine/Plugins/Experimental/OperatorStack/Source/OperatorStack/Private/Styles/OperatorStackEditorStyle.h
@@ -0,0 +1,19 @@
+// Copyright Epic Games, Inc. All Rights Reserved.
+
+#pragma once
+
+#include "Styling/SlateStyle.h"
+
+class FOperatorStackEditorStyle : public FSlateStyleSet
+{
+public:
+	FOperatorStackEditorStyle();
+	
+	virtual ~FOperatorStackEditorStyle() override;
+	
+	static FOperatorStackEditorStyle& Get()
+	{
+		static FOperatorStackEditorStyle StyleSet;
+		return StyleSet;
+	}
+};
diff --git a/Engine/Plugins/Experimental/OperatorStack/Source/OperatorStack/Private/Subsystems/OperatorStackEditorSubsystem.cpp b/Engine/Plugins/Experimental/OperatorStack/Source/OperatorStack/Private/Subsystems/OperatorStackEditorSubsystem.cpp
new file mode 100644
--- /dev/null
+++ b/Engine/Plugins/Experimental/OperatorStack/Source/OperatorStack/Private/Subsystems/OperatorStackEditorSubsystem.cpp
@@ -0,0 +1,187 @@
+// Copyright Epic Games, Inc. All Rights Reserved.
+
+#include "Subsystems/OperatorStackEditorSubsystem.h"
+
+#include "Editor.h"
+#include "Customizations/OperatorStackEditorStackCustomization.h"
+#include "UObject/UObjectIterator.h"
+#include "Widgets/SOperatorStackEditorPanel.h"
+#include "Widgets/SOperatorStackEditorWidget.h"
+
+DEFINE_LOG_CATEGORY_STATIC(LogOperatorStackEditorSubsystem, Log, All);
+
+UOperatorStackEditorSubsystem* UOperatorStackEditorSubsystem::Get()
+{
+	if (GEditor)
+	{
+		return GEditor->GetEditorSubsystem<UOperatorStackEditorSubsystem>();
+	}
+	return nullptr;
+}
+
+void UOperatorStackEditorSubsystem::Initialize(FSubsystemCollectionBase& Collection)
+{
+	Super::Initialize(Collection);
+
+	ScanForStackCustomizations();
+}
+
+void UOperatorStackEditorSubsystem::Deinitialize()
+{
+	CustomizationStacks.Empty();
+
+	Super::Deinitialize();
+}
+
+bool UOperatorStackEditorSubsystem::RegisterStackCustomization(const TSubclassOf<UOperatorStackEditorStackCustomization> InStackCustomizationClass)
+{
+	if (!InStackCustomizationClass->IsChildOf(UOperatorStackEditorStackCustomization::StaticClass())
+		|| InStackCustomizationClass->HasAnyClassFlags(CLASS_Abstract | CLASS_Deprecated | CLASS_NewerVersionExists))
+	{
+		return false;
+	}
+
+	UOperatorStackEditorStackCustomization* StackCustomization = InStackCustomizationClass->GetDefaultObject<UOperatorStackEditorStackCustomization>();
+	if (!StackCustomization)
+	{
+		return false;
+	}
+
+	const FName StackIdentifier = StackCustomization->GetIdentifier();
+
+	if (CustomizationStacks.Contains(StackIdentifier))
+	{
+		return false;
+	}
+
+	CustomizationStacks.Add(StackIdentifier, StackCustomization);
+
+	UE_LOG(LogOperatorStackEditorSubsystem, Display, TEXT("OperatorStack customization registered : Class %s - Identifier %s"), *InStackCustomizationClass->GetName(), *StackIdentifier.ToString());
+
+	return true;
+}
+
+bool UOperatorStackEditorSubsystem::UnregisterStackCustomization(const TSubclassOf<UOperatorStackEditorStackCustomization> InStackCustomizationClass)
+{
+	const UOperatorStackEditorStackCustomization* StackCustomization = InStackCustomizationClass->GetDefaultObject<UOperatorStackEditorStackCustomization>();
+	if (!StackCustomization)
+	{
+		return false;
+	}
+
+	const FName StackIdentifier = StackCustomization->GetIdentifier();
+
+    if (!CustomizationStacks.Contains(StackIdentifier))
+    {
+    	return false;
+    }
+
+    CustomizationStacks.Remove(StackIdentifier);
+
+	UE_LOG(LogOperatorStackEditorSubsystem, Display, TEXT("OperatorStack customization unregistered : Class %s - Identifier %s"), *InStackCustomizationClass->GetName(), *StackIdentifier.ToString());
+
+    return true;
+}
+
+TSharedRef<SOperatorStackEditorWidget> UOperatorStackEditorSubsystem::GenerateWidget()
+{
+	// Find new id available
+	int32 NewId = 0;
+	while (CustomizationWidgets.Contains(NewId))
+	{
+		NewId++;
+	}
+
+	TSharedRef<SOperatorStackEditorWidget> NewWidget = SNew(SOperatorStackEditorPanel)
+		.PanelId(NewId);
+
+	CustomizationWidgets.Add(NewId, NewWidget);
+
+	return NewWidget;
+}
+
+TSharedPtr<SOperatorStackEditorWidget> UOperatorStackEditorSubsystem::FindWidget(int32 InId)
+{
+	if (const TWeakPtr<SOperatorStackEditorWidget>* WidgetWeak = CustomizationWidgets.Find(InId))
+	{
+		return WidgetWeak->Pin();
+	}
+
+	return nullptr;
+}
+
+bool UOperatorStackEditorSubsystem::ForEachCustomization(TFunctionRef<bool(UOperatorStackEditorStackCustomization*)> InFunction) const
+{
+	TArray<UOperatorStackEditorStackCustomization*> Customizations;
+
+	for (const TPair<FName, TObjectPtr<UOperatorStackEditorStackCustomization>>& CustomizationPair : CustomizationStacks)
+	{
+		if (IsValid(CustomizationPair.Value))
+		{
+			Customizations.Add(CustomizationPair.Value);
+		}
+	}
+
+	Customizations.Sort([](const UOperatorStackEditorStackCustomization& InA, const UOperatorStackEditorStackCustomization& InB)
+	{
+		return InA.GetPriority() > InB.GetPriority();
+	});
+
+	for (UOperatorStackEditorStackCustomization* Customization : Customizations)
+	{
+		if (!InFunction(Customization))
+		{
+			return false;
+		}
+	}
+
+	return true;
+}
+
+bool UOperatorStackEditorSubsystem::ForEachCustomizationWidget(TFunctionRef<bool(TSharedRef<SOperatorStackEditorWidget>)> InFunction) const
+{
+	TArray<TWeakPtr<SOperatorStackEditorWidget>> WeakWidgets;
+	CustomizationWidgets.GenerateValueArray(WeakWidgets);
+
+	for (TWeakPtr<SOperatorStackEditorWidget>& WeakWidget : WeakWidgets)
+	{
+		TSharedPtr<SOperatorStackEditorWidget> Widget = WeakWidget.Pin();
+
+		if (Widget.IsValid())
+		{
+			if (!InFunction(Widget.ToSharedRef()))
+			{
+				return false;
+			}
+		}
+	}
+
+	return true;
+}
+
+UOperatorStackEditorStackCustomization* UOperatorStackEditorSubsystem::GetCustomization(const FName& InName) const
+{
+	if (const TObjectPtr<UOperatorStackEditorStackCustomization>* Customization = CustomizationStacks.Find(InName))
+	{
+		return *Customization;
+	}
+
+	return nullptr;
+}
+
+void UOperatorStackEditorSubsystem::ScanForStackCustomizations()
+{
+	for (UClass* const Class : TObjectRange<UClass>())
+	{
+		if (Class && Class->IsChildOf(UOperatorStackEditorStackCustomization::StaticClass()))
+		{
+			const TSubclassOf<UOperatorStackEditorStackCustomization> ScannedClass(Class);
+			RegisterStackCustomization(ScannedClass);
+		}
+	}
+}
+
+void UOperatorStackEditorSubsystem::OnWidgetDestroyed(int32 InPanelId)
+{
+	CustomizationWidgets.Remove(InPanelId);
+}
diff --git a/Engine/Plugins/Experimental/OperatorStack/Source/OperatorStack/Private/Widgets/DragDropOps/OperatorStackEditorDragDropOp.h b/Engine/Plugins/Experimental/OperatorStack/Source/OperatorStack/Private/Widgets/DragDropOps/OperatorStackEditorDragDropOp.h
new file mode 100644
--- /dev/null
+++ b/Engine/Plugins/Experimental/OperatorStack/Source/OperatorStack/Private/Widgets/DragDropOps/OperatorStackEditorDragDropOp.h
@@ -0,0 +1,50 @@
+// Copyright Epic Games, Inc. All Rights Reserved.
+
+#pragma once
+
+#include "DragAndDrop/DecoratedDragDropOp.h"
+#include "Misc/Optional.h"
+#include "Templates/SharedPointer.h"
+
+enum class EItemDropZone;
+struct FOperatorStackEditorItem;
+
+/** A custom drag operation for operator stack items */
+class FOperatorStackEditorDragDropOp : public FDecoratedDragDropOp
+{
+public:
+	DRAG_DROP_OPERATOR_TYPE(FOperatorStackEditorDragDropOp, FDecoratedDragDropOp)
+
+	static TSharedRef<FOperatorStackEditorDragDropOp> New(const TArray<TSharedPtr<FOperatorStackEditorItem>>& InItems)
+	{
+		TSharedRef<FOperatorStackEditorDragDropOp> DragDropOp = MakeShared<FOperatorStackEditorDragDropOp>();
+		DragDropOp->DraggedItems = InItems;
+		DragDropOp->CurrentIconColorAndOpacity = FSlateColor::UseForeground();
+		DragDropOp->MouseCursor = EMouseCursor::GrabHandClosed;
+
+		DragDropOp->SetupDefaults();
+		DragDropOp->Construct();
+		
+		return DragDropOp;
+	}
+
+	const TArray<TSharedPtr<FOperatorStackEditorItem>>& GetDraggedItems() const
+	{
+		return DraggedItems;
+	}
+
+	void SetDropZone(const TOptional<EItemDropZone>& InDropZone)
+	{
+		DropZone = InDropZone;	
+	}
+
+	TOptional<EItemDropZone> GetDropZone() const
+	{
+		return DropZone;	
+	}
+protected:
+	/** Dragged items */
+	TArray<TSharedPtr<FOperatorStackEditorItem>> DraggedItems;
+
+	TOptional<EItemDropZone> DropZone;
+};
diff --git a/Engine/Plugins/Experimental/OperatorStack/Source/OperatorStack/Private/Widgets/SOperatorStackEditorPanel.cpp b/Engine/Plugins/Experimental/OperatorStack/Source/OperatorStack/Private/Widgets/SOperatorStackEditorPanel.cpp
new file mode 100644
--- /dev/null
+++ b/Engine/Plugins/Experimental/OperatorStack/Source/OperatorStack/Private/Widgets/SOperatorStackEditorPanel.cpp
@@ -0,0 +1,316 @@
+// Copyright Epic Games, Inc. All Rights Reserved.
+
+#include "Widgets/SOperatorStackEditorPanel.h"
+
+#include "DetailColumnSizeData.h"
+#include "Customizations/OperatorStackEditorStackCustomization.h"
+#include "SOperatorStackEditorStack.h"
+#include "Styling/ToolBarStyle.h"
+#include "Subsystems/OperatorStackEditorSubsystem.h"
+#include "Widgets/Images/SImage.h"
+#include "Widgets/Input/SCheckBox.h"
+#include "Widgets/Layout/SBox.h"
+#include "Widgets/Layout/SScrollBox.h"
+#include "Widgets/Layout/SSeparator.h"
+#include "Widgets/Layout/SWidgetSwitcher.h"
+#include "Widgets/Text/STextBlock.h"
+
+SOperatorStackEditorPanel::~SOperatorStackEditorPanel()
+{
+	if (UOperatorStackEditorSubsystem* EditorSubsystem = UOperatorStackEditorSubsystem::Get())
+	{
+		EditorSubsystem->OnWidgetDestroyed(PanelId);
+	}
+}
+
+void SOperatorStackEditorPanel::Construct(const FArguments& InArgs)
+{
+	static const FToolBarStyle& ToolBarStyle = FAppStyle::Get().GetWidgetStyle<FToolBarStyle>("SlimToolBar");
+
+	PanelId = InArgs._PanelId;
+
+	// Generate default empty context
+	Context = MakeShared<FOperatorStackEditorContext>();
+
+	// Setup default column sizes
+	DetailColumnSize = MakeShared<FDetailColumnSizeData>();
+	DetailColumnSize->SetRightColumnMinWidth(60.0f);
+	DetailColumnSize->SetValueColumnWidth(0.6f);
+
+	ChildSlot
+	[
+		SNew(SVerticalBox)
+		// Toolbar to switch between different stacks
+		+ SVerticalBox::Slot()
+		.AutoHeight()
+		[
+			SNew(SBox)
+			.HAlign(HAlign_Center)
+			.VAlign(VAlign_Fill)
+			.Padding(0.f, SOperatorStackEditorStack::Padding)
+			[
+				SAssignNew(WidgetToolbar, SScrollBox)
+				.Orientation(EOrientation::Orient_Horizontal)
+				.ScrollBarThickness(FVector2D(1.f))
+			]
+		]
+		+ SVerticalBox::Slot()
+		.AutoHeight()
+		[
+			SNew(SSeparator)
+			.Orientation(Orient_Horizontal)
+			.Thickness(3.f)
+		]
+		// Widget for each stacks
+		+ SVerticalBox::Slot()
+		.FillHeight(1.f)
+		[
+			SAssignNew(WidgetSwitcher, SWidgetSwitcher)
+		]
+	];
+
+	const UOperatorStackEditorSubsystem* EditorSubsystem = UOperatorStackEditorSubsystem::Get();
+	EditorSubsystem->ForEachCustomization([this](UOperatorStackEditorStackCustomization* InCustomization)
+	{
+		AddSlot(InCustomization);
+		return true;
+	});
+
+	UpdateSlots();
+}
+
+void SOperatorStackEditorPanel::SetContext(const FOperatorStackEditorContext& InContext)
+{
+	Context = MakeShared<FOperatorStackEditorContext>(InContext);
+	RefreshContext();
+}
+
+void SOperatorStackEditorPanel::SetActiveCustomization(const FName& InCustomization)
+{
+	if (WidgetSwitcher.IsValid() && NamedStackIndexes.Contains(InCustomization))
+	{
+		const int32* WidgetIdx = NamedStackIndexes.Find(InCustomization);
+		WidgetSwitcher->SetActiveWidgetIndex(*WidgetIdx);
+	}
+}
+
+void SOperatorStackEditorPanel::SetToolbarCustomization(const TArray<FName>& InCustomizations)
+{
+	if (!WidgetToolbar.IsValid() || !WidgetSwitcher.IsValid())
+	{
+		return;
+	}
+
+	WidgetToolbar->ClearChildren();
+	for (int32 SlotIdx = 0; SlotIdx < WidgetSwitcher->GetNumWidgets(); SlotIdx++)
+	{
+		if (TSharedPtr<SWidget> SlotWidget = WidgetSwitcher->GetWidget(SlotIdx))
+		{
+			WidgetSwitcher->RemoveSlot(SlotWidget.ToSharedRef());
+		}
+	}
+
+	const UOperatorStackEditorSubsystem* EditorSubsystem = UOperatorStackEditorSubsystem::Get();
+
+	for (const FName& CustomizationName : InCustomizations)
+	{
+		const UOperatorStackEditorStackCustomization* Customization = EditorSubsystem->GetCustomization(CustomizationName);
+
+		if (!Customization)
+		{
+			continue;
+		}
+
+		AddSlot(Customization);
+	}
+
+	UpdateSlots();
+}
+
+void SOperatorStackEditorPanel::ShowToolbarCustomization(FName InCustomization)
+{
+	if (const int32* CustomizationIdx = NamedStackIndexes.Find(InCustomization))
+	{
+		VisibleCustomizations.Add(*CustomizationIdx);
+	}
+}
+
+void SOperatorStackEditorPanel::HideToolbarCustomization(FName InCustomization)
+{
+	if (const int32* CustomizationIdx = NamedStackIndexes.Find(InCustomization))
+	{
+		VisibleCustomizations.Remove(*CustomizationIdx);
+	}
+}
+
+void SOperatorStackEditorPanel::SetToolbarVisibility(bool bInVisible)
+{
+	if (!WidgetToolbar.IsValid())
+	{
+		return;
+	}
+
+	WidgetToolbar->SetVisibility(bInVisible ? EVisibility::Visible : EVisibility::Collapsed);
+}
+
+void SOperatorStackEditorPanel::SetKeyframeHandler(TSharedPtr<IDetailKeyframeHandler> InKeyframeHandler)
+{
+	KeyframeHandler = InKeyframeHandler;
+	RefreshContext();
+}
+
+void SOperatorStackEditorPanel::SetDetailColumnSize(TSharedPtr<FDetailColumnSizeData> InDetailColumnSize)
+{
+	DetailColumnSize = InDetailColumnSize;
+	RefreshContext();
+}
+
+void SOperatorStackEditorPanel::SetPanelTag(FName InTag)
+{
+	PanelTag = InTag;
+}
+
+void SOperatorStackEditorPanel::RefreshContext()
+{
+	UpdateSlots();
+}
+
+const FOperatorStackEditorTree& SOperatorStackEditorPanel::GetItemTree(UOperatorStackEditorStackCustomization* InCustomization)
+{
+	const FOperatorStackEditorTree* Tree = CustomizationTrees.Find(InCustomization);
+	check(Tree)
+	return *Tree;
+}
+
+void SOperatorStackEditorPanel::SaveItemExpansionState(const void* InItem, bool bInExpanded)
+{
+	ItemExpansionState.Add(InItem, bInExpanded);
+}
+
+bool SOperatorStackEditorPanel::GetItemExpansionState(const void* InItem, bool& bOutExpanded)
+{
+	if (const bool* bExpanded = ItemExpansionState.Find(InItem))
+	{
+		bOutExpanded = *bExpanded;
+		return true;
+	}
+
+	return false;
+}
+
+EVisibility SOperatorStackEditorPanel::GetToolbarButtonVisibility(int32 InIdx) const
+{
+	return VisibleCustomizations.Contains(InIdx) ? EVisibility::Visible : EVisibility::Collapsed;
+}
+
+void SOperatorStackEditorPanel::AddSlot(const UOperatorStackEditorStackCustomization* InCustomizationStack)
+{
+	if (!WidgetToolbar.IsValid() || !WidgetSwitcher.IsValid())
+	{
+		return;
+	}
+
+	static const FToolBarStyle& ToolBarStyle = FAppStyle::Get().GetWidgetStyle<FToolBarStyle>("SlimToolBar");
+	static const FCheckBoxStyle* const CheckStyle = &ToolBarStyle.ToggleButton;
+
+	const FText& Label = InCustomizationStack->GetLabel();
+	const FName& Identifier = InCustomizationStack->GetIdentifier();
+	const FSlateBrush* Icon = InCustomizationStack->GetIcon();
+
+	// Add widget for stack content
+	WidgetSwitcher->AddSlot()
+		.HAlign(HAlign_Fill)
+		.VAlign(VAlign_Fill)
+		.Padding(SOperatorStackEditorStack::Padding, 0.f)
+		[
+			SNew(SBox)
+			.Padding(0.f)
+			.HAlign(HAlign_Fill)
+			.VAlign(VAlign_Fill)
+		];
+
+	int32 WidgetIdx = WidgetSwitcher->GetNumWidgets() - 1;
+
+	// Add button to switch to new stack
+	WidgetToolbar->AddSlot()
+		.Padding(SOperatorStackEditorStack::Padding)
+		[
+			SNew(SCheckBox)
+			.Style(CheckStyle)
+			.ForegroundColor(FLinearColor::White)
+			.OnCheckStateChanged(this, &SOperatorStackEditorPanel::OnToolbarButtonClicked, WidgetIdx)
+			.IsChecked(this, &SOperatorStackEditorPanel::IsToolbarButtonActive, WidgetIdx)
+			.Visibility(this, &SOperatorStackEditorPanel::GetToolbarButtonVisibility, WidgetIdx)
+			[
+				SNew(SHorizontalBox)
+				+ SHorizontalBox::Slot()
+				.AutoWidth()
+				.Padding(2.f)
+				.HAlign(HAlign_Left)
+				.VAlign(VAlign_Center)
+				[
+					SNew(SImage)
+					.DesiredSizeOverride(FVector2D(12.f, 12.f))
+					.Image(Icon)
+				]
+				+ SHorizontalBox::Slot()
+				.FillWidth(1.f)
+				.Padding(2.f)
+				.HAlign(HAlign_Left)
+				.VAlign(VAlign_Center)
+				[
+					SNew(STextBlock)
+					.Text(Label)
+				]
+			]
+		];
+
+	NamedStackIndexes.Emplace(Identifier, WidgetIdx);
+	VisibleCustomizations.Add(WidgetIdx);
+}
+
+void SOperatorStackEditorPanel::UpdateSlots()
+{
+	const UOperatorStackEditorSubsystem* EditorSubsystem = UOperatorStackEditorSubsystem::Get();
+
+	for (const TPair<FName, int32>& NamedSlot : NamedStackIndexes)
+	{
+		if (const TSharedPtr<SBox> Box = StaticCastSharedPtr<SBox>(WidgetSwitcher->GetWidget(NamedSlot.Value)))
+		{
+			UOperatorStackEditorStackCustomization* Customization = EditorSubsystem->GetCustomization(NamedSlot.Key);
+
+			check(Customization);
+
+			CustomizationTrees.Add(Customization, FOperatorStackEditorTree(Customization, Context));
+
+			Box->SetContent(
+				SNew(SOperatorStackEditorStack
+					, SharedThis(this)
+					, Customization
+					, nullptr)
+			);
+		}
+	}
+}
+
+void SOperatorStackEditorPanel::OnToolbarButtonClicked(ECheckBoxState InState, int32 InWidgetIdx) const
+{
+	if (!WidgetSwitcher.IsValid() || InWidgetIdx == INDEX_NONE)
+	{
+		return;
+	}
+
+	WidgetSwitcher->SetActiveWidgetIndex(InWidgetIdx);
+}
+
+ECheckBoxState SOperatorStackEditorPanel::IsToolbarButtonActive(int32 InWidgetIdx) const
+{
+	if (!WidgetSwitcher.IsValid() || InWidgetIdx == INDEX_NONE)
+	{
+		return ECheckBoxState::Undetermined;
+	}
+
+	return WidgetSwitcher->GetActiveWidgetIndex() == InWidgetIdx
+		? ECheckBoxState::Checked
+		: ECheckBoxState::Unchecked;
+}
diff --git a/Engine/Plugins/Experimental/OperatorStack/Source/OperatorStack/Private/Widgets/SOperatorStackEditorPanel.h b/Engine/Plugins/Experimental/OperatorStack/Source/OperatorStack/Private/Widgets/SOperatorStackEditorPanel.h
new file mode 100644
--- /dev/null
+++ b/Engine/Plugins/Experimental/OperatorStack/Source/OperatorStack/Private/Widgets/SOperatorStackEditorPanel.h
@@ -0,0 +1,131 @@
+// Copyright Epic Games, Inc. All Rights Reserved.
+
+#pragma once
+
+#include "CustomDetailsViewArgs.h"
+#include "Contexts/OperatorStackEditorContext.h"
+#include "Items/OperatorStackEditorTree.h"
+#include "Templates/SharedPointer.h"
+#include "Widgets/SCompoundWidget.h"
+#include "Widgets/SOperatorStackEditorWidget.h"
+
+class FReply;
+class SOperatorStackEditorStack;
+class SScrollBox;
+class SWidgetSwitcher;
+class UOperatorStackEditorStackCustomization;
+struct FButtonStyle;
+
+/** This is the main panel that contains the stacks and switcher for them */
+class SOperatorStackEditorPanel : public SOperatorStackEditorWidget
+{
+public:
+	SLATE_BEGIN_ARGS(SOperatorStackEditorPanel) {}
+		SLATE_ARGUMENT(int32, PanelId)
+	SLATE_END_ARGS()
+
+	virtual ~SOperatorStackEditorPanel() override;
+
+	void Construct(const FArguments& InArgs);
+
+	virtual void SetContext(const FOperatorStackEditorContext& InContext) override;
+
+	virtual void SetActiveCustomization(const FName& InCustomization) override;
+
+	virtual void SetToolbarCustomization(const TArray<FName>& InCustomizations) override;
+
+	virtual void ShowToolbarCustomization(FName InCustomization) override;
+
+	virtual void HideToolbarCustomization(FName InCustomization) override;
+
+	virtual void SetToolbarVisibility(bool bInVisible) override;
+
+	virtual void SetKeyframeHandler(TSharedPtr<IDetailKeyframeHandler> InKeyframeHandler) override;
+
+	virtual void SetDetailColumnSize(TSharedPtr<FDetailColumnSizeData> InDetailColumnSize) override;
+
+	virtual void SetPanelTag(FName InTag) override;
+
+	virtual void RefreshContext() override;
+
+	const FOperatorStackEditorTree& GetItemTree(UOperatorStackEditorStackCustomization* InCustomization);
+
+	void SaveItemExpansionState(const void* InItem, bool bInExpanded);
+
+	bool GetItemExpansionState(const void* InItem, bool& bOutExpanded);
+
+	virtual FOperatorStackEditorContextPtr GetContext() const override
+	{
+		return Context;
+	}
+
+	virtual int32 GetPanelId() const override
+	{
+		return PanelId;
+	}
+
+	virtual FName GetPanelTag() const override
+	{
+		return PanelTag;
+	}
+
+	TSharedPtr<IDetailKeyframeHandler> GetKeyframeHandler() const
+	{
+		return KeyframeHandler;
+	}
+
+	TSharedPtr<FDetailColumnSizeData> GetDetailColumnSize() const
+	{
+		return DetailColumnSize;
+	}
+
+protected:
+	/** Is the toolbar button active */
+	EVisibility GetToolbarButtonVisibility(int32 InIdx) const;
+
+	/** Adds a new widget slot for toolbar and switcher */
+	void AddSlot(const UOperatorStackEditorStackCustomization* InCustomizationStack);
+
+	/** Update the slots content with actual widget */
+	void UpdateSlots();
+
+	/** Switch view when toolbar button is clicked */
+	void OnToolbarButtonClicked(ECheckBoxState InState, int32 InWidgetIdx) const;
+
+	/** Highlight active selection */
+	ECheckBoxState IsToolbarButtonActive(int32 InWidgetIdx) const;
+
+private:
+	/** Contains the named stack views */
+	TSharedPtr<SScrollBox> WidgetToolbar;
+
+	/** Contains the different stack views */
+	TSharedPtr<SWidgetSwitcher> WidgetSwitcher;
+
+	/** Stack index with their identifier for updates */
+	TMap<FName, int32> NamedStackIndexes;
+
+	/** Customization visible to switch in toolbar */
+	TSet<int32> VisibleCustomizations;
+
+	/** Current context, could be null */
+	FOperatorStackEditorContextPtr Context;
+
+	/** Keyframe handler to be able to keyframe properties in items */
+	TSharedPtr<IDetailKeyframeHandler> KeyframeHandler;
+
+	/** Detail column size to have matching size across items */
+	TSharedPtr<FDetailColumnSizeData> DetailColumnSize;
+
+	/** The panel id to identify this widget */
+	int32 PanelId = INDEX_NONE;
+
+	/** The panel tag to identify this widget */
+	FName PanelTag = NAME_None;
+
+	/** Item trees containing items supported for each customization */
+	TMap<TWeakObjectPtr<UOperatorStackEditorStackCustomization>, FOperatorStackEditorTree> CustomizationTrees;
+
+	/** To check whether an item was expanded and restore it */
+	TMap<const void*, bool> ItemExpansionState;
+};
\ No newline at end of file
diff --git a/Engine/Plugins/Experimental/OperatorStack/Source/OperatorStack/Private/Widgets/SOperatorStackEditorStack.cpp b/Engine/Plugins/Experimental/OperatorStack/Source/OperatorStack/Private/Widgets/SOperatorStackEditorStack.cpp
new file mode 100644
--- /dev/null
+++ b/Engine/Plugins/Experimental/OperatorStack/Source/OperatorStack/Private/Widgets/SOperatorStackEditorStack.cpp
@@ -0,0 +1,1046 @@
+// Copyright Epic Games, Inc. All Rights Reserved.
+
+#include "Widgets/SOperatorStackEditorStack.h"
+
+#include "Builders/OperatorStackEditorBodyBuilder.h"
+#include "Builders/OperatorStackEditorFooterBuilder.h"
+#include "Builders/OperatorStackEditorHeaderBuilder.h"
+#include "CustomDetailsViewModule.h"
+#include "Customizations/OperatorStackEditorStackCustomization.h"
+#include "Framework/Application/SlateApplication.h"
+#include "ICustomDetailsView.h"
+#include "SOperatorStackEditorPanel.h"
+#include "Items/ICustomDetailsViewItem.h"
+#include "SOperatorStackExpanderButton.h"
+#include "Styles/OperatorStackEditorStyle.h"
+#include "Styling/AppStyle.h"
+#include "Styling/ToolBarStyle.h"
+#include "ToolMenu.h"
+#include "ToolMenus.h"
+#include "Contexts/OperatorStackEditorMenuContext.h"
+#include "Items/OperatorStackEditorTree.h"
+#include "Items/OperatorStackEditorStructItem.h"
+#include "Widgets/Colors/SColorBlock.h"
+#include "Widgets/Input/SCheckBox.h"
+#include "Widgets/Input/SComboButton.h"
+#include "Widgets/Input/SSearchBox.h"
+#include "Widgets/Layout/SSeparator.h"
+#include "Widgets/Layout/SScaleBox.h"
+#include "Widgets/Layout/SScrollBox.h"
+#include "Widgets/TableRows/SOperatorStackEditorStackRow.h"
+
+#define LOCTEXT_NAMESPACE "OperatorStackEditorStack"
+
+const TMap<EOperatorStackEditorMessageType, FLinearColor> SOperatorStackEditorStack::MessageBoxColors
+{
+	{EOperatorStackEditorMessageType::None, FLinearColor::Transparent},
+	{EOperatorStackEditorMessageType::Info, FLinearColor::Blue.Desaturate(0.5)},
+	{EOperatorStackEditorMessageType::Success, FLinearColor::Green.Desaturate(0.5)},
+	{EOperatorStackEditorMessageType::Warning, FLinearColor::Yellow.Desaturate(0.5)},
+	{EOperatorStackEditorMessageType::Error, FLinearColor::Red.Desaturate(0.5)},
+};
+
+const TMap<EOperatorStackEditorMessageType, const FSlateBrush*> SOperatorStackEditorStack::MessageBoxIcons
+{
+	{EOperatorStackEditorMessageType::None, nullptr},
+	{EOperatorStackEditorMessageType::Info, FAppStyle::GetBrush("Icons.InfoWithColor")},
+	{EOperatorStackEditorMessageType::Success, FAppStyle::GetBrush("Icons.SuccessWithColor")},
+	{EOperatorStackEditorMessageType::Warning, FAppStyle::GetBrush("Icons.WarningWithColor")},
+	{EOperatorStackEditorMessageType::Error, FAppStyle::GetBrush("Icons.ErrorWithColor")},
+};
+
+void SOperatorStackEditorStack::Construct(const FArguments& InArgs
+    , const TSharedPtr<SOperatorStackEditorPanel> InMainPanel
+    , TObjectPtr<UOperatorStackEditorStackCustomization> InCustomization
+    , const FOperatorStackEditorItemPtr& InCustomizeItem)
+{
+	MainPanelWeak = InMainPanel;
+	CustomizeItem = InCustomizeItem;
+	StackCustomizationWeak = InCustomization;
+
+	check(InCustomization && InMainPanel);
+
+	const FOperatorStackEditorTree& ItemTree = InMainPanel->GetItemTree(InCustomization);
+
+	if (InCustomizeItem.IsValid())
+	{
+		Items = ItemTree.GetChildrenItems(InCustomizeItem);
+	}
+	else
+	{
+		Items = ItemTree.GetRootItems();
+	}
+
+	ChildSlot
+	.HAlign(HAlign_Fill)
+	.VAlign(VAlign_Fill)
+	[
+		GenerateStackWidget()
+	];
+}
+
+FOperatorStackEditorContextPtr SOperatorStackEditorStack::GetContext() const
+{
+	if (const TSharedPtr<SOperatorStackEditorPanel> MainPanel = GetMainPanel())
+	{
+		return MainPanel->GetContext();
+	}
+
+	return nullptr;
+}
+
+TSharedRef<ITableRow> SOperatorStackEditorStack::OnGenerateRow(FOperatorStackEditorItemPtr InItem, const TSharedRef<STableViewBase>& InOwnerTable)
+{
+	return SNew(SOperatorStackEditorStackRow, InOwnerTable, SharedThis(this), InItem);
+}
+
+TSharedPtr<SWidget> SOperatorStackEditorStack::GenerateHeaderWidget()
+{
+	UOperatorStackEditorStackCustomization* StackCustomization = GetStackCustomization();
+	const FOperatorStackEditorContextPtr Context = GetContext();
+	const TSharedPtr<SOperatorStackEditorPanel> MainPanel = GetMainPanel();
+
+	if (!Context.IsValid())
+	{
+		return nullptr;
+	}
+
+	check(MainPanel.IsValid());
+
+	FOperatorStackEditorHeaderBuilder HeaderBuilder;
+	const FOperatorStackEditorTree& ItemTree = MainPanel->GetItemTree(StackCustomization);
+
+	// We are in the top most stack, customize stack header
+	if (!CustomizeItem.IsValid())
+	{
+		StackCustomization->CustomizeStackHeader(ItemTree, HeaderBuilder);
+	}
+	// We are in a child item, customize child header
+	else
+	{
+		StackCustomization->CustomizeItemHeader(CustomizeItem, ItemTree, HeaderBuilder);
+	}
+
+	// Get Commands for this item
+	CommandList = HeaderBuilder.GetCommandList();
+
+	// Get Context menu name for this item
+	ContextMenuName = HeaderBuilder.GetContextMenuName();
+
+	// Get border color for item
+	BorderColor = HeaderBuilder.GetBorderColor();
+
+	const TSharedRef<SHorizontalBox> HorizontalHeaderWidget = SNew(SHorizontalBox);
+
+	// Expansion button to show body and footer
+	if (HeaderBuilder.GetExpandable())
+	{
+		// Find previous expansion state or use default
+		bHeaderExpanded = HeaderBuilder.GetStartsExpanded();
+		MainPanel->GetItemExpansionState(CustomizeItem->GetValuePtr(), bHeaderExpanded);
+
+		HorizontalHeaderWidget->AddSlot()
+		.AutoWidth()
+		.Padding(Padding)
+		.HAlign(HAlign_Left)
+		.VAlign(VAlign_Center)
+		[
+			SNew(SOperatorStackExpanderButton)
+			.StartsExpanded(bHeaderExpanded)
+			.OnExpansionStateChanged(this, &SOperatorStackEditorStack::OnHeaderExpansionChanged)
+		];
+	}
+
+	// Icon
+	if (const FSlateBrush* Icon = HeaderBuilder.GetIcon())
+	{
+		HorizontalHeaderWidget->AddSlot()
+		.AutoWidth()
+		.Padding(Padding)
+		.HAlign(HAlign_Left)
+		.VAlign(VAlign_Center)
+		[
+			SNew(SImage)
+			.Image(Icon)
+		];
+	}
+
+	// Label
+	if (!HeaderBuilder.GetLabel().IsEmpty())
+	{
+		HorizontalHeaderWidget->AddSlot()
+		.AutoWidth()
+		.Padding(Padding)
+		.HAlign(HAlign_Left)
+		.VAlign(VAlign_Center)
+		[
+			SNew(STextBlock)
+			.Text(HeaderBuilder.GetLabel())
+			.Justification(ETextJustify::Center)
+		];
+	}
+
+	// Space here
+	HorizontalHeaderWidget->AddSlot()
+		.FillWidth(1.f)
+		.Padding(Padding)
+		.HAlign(HAlign_Fill)
+		.VAlign(VAlign_Fill)
+		[
+			SNullWidget::NullWidget
+		];
+
+	// custom actions menu in header
+	const FName& ActionMenuName = HeaderBuilder.GetToolbarMenuName();
+	const UToolMenu* ActionMenu = !ActionMenuName.IsNone()
+		? UToolMenus::Get()->FindMenu(ActionMenuName)
+		: nullptr;
+
+	if (ActionMenu)
+	{
+		HorizontalHeaderWidget->AddSlot()
+		.Padding(Padding)
+		.HAlign(HAlign_Right)
+		.VAlign(VAlign_Center)
+		[
+			SNew(SBox)
+			.HeightOverride(25.f)
+			[
+				SNew(SScaleBox)
+				.Stretch(EStretch::ScaleToFit)
+				[
+					GenerateMenuWidget(ActionMenuName)
+				]
+			]
+		];
+	}
+
+	// Custom property in header
+	const TSharedPtr<FCustomDetailsViewItemId> PropertyId = HeaderBuilder.GetProperty();
+	if (CustomizeItem.IsValid() && PropertyId.IsValid() && PropertyId->IsValid())
+	{
+		const FCustomDetailsViewItemId& ItemId = *PropertyId.Get();
+
+		FCustomDetailsViewArgs HeaderCustomViewArgs;
+
+        HeaderCustomViewArgs.bShowCategories = false;
+        HeaderCustomViewArgs.KeyframeHandler = MainPanel->GetKeyframeHandler();
+        HeaderCustomViewArgs.bAllowGlobalExtensions = true;
+        HeaderCustomViewArgs.ItemAllowList.Allow(ItemId);
+        HeaderCustomViewArgs.WidgetTypeAllowList.Allow(ECustomDetailsViewWidgetType::Value);
+        HeaderCustomViewArgs.WidgetTypeAllowList.Allow(ECustomDetailsViewWidgetType::Extensions);
+
+		HeaderDetailsView = CreateDetailsView(HeaderCustomViewArgs, *CustomizeItem);
+
+		if (const TSharedPtr<ICustomDetailsViewItem> HeaderCustomDetailsViewItem = HeaderDetailsView->FindItem(ItemId))
+		{
+			HeaderCustomDetailsViewItem->MakeWidget(nullptr, nullptr);
+
+			// Add Property header widget (eg: checkbox for bool)
+			const TSharedPtr<SWidget> HeaderPropertyValueWidget = HeaderCustomDetailsViewItem->GetWidget(ECustomDetailsViewWidgetType::Value);
+			if (HeaderPropertyValueWidget.IsValid())
+			{
+				HorizontalHeaderWidget->AddSlot()
+				.AutoWidth()
+				.MaxWidth(25.0f)
+				.Padding(Padding)
+				.HAlign(HAlign_Right)
+				.VAlign(VAlign_Center)
+				[
+					HeaderPropertyValueWidget.ToSharedRef()
+				];
+			}
+
+			// Add Expansion property widget (eg : keyframe icon)
+			const TSharedPtr<SWidget> HeaderPropertyExtensionWidget = HeaderCustomDetailsViewItem->GetWidget(ECustomDetailsViewWidgetType::Extensions);
+			if (HeaderPropertyExtensionWidget.IsValid())
+			{
+				HorizontalHeaderWidget->AddSlot()
+				.AutoWidth()
+				.MaxWidth(25.0f)
+				.Padding(Padding)
+				.HAlign(HAlign_Right)
+				.VAlign(VAlign_Center)
+				[
+					HeaderPropertyExtensionWidget.ToSharedRef()
+				];
+			}
+		}
+	}
+
+	const TSharedPtr<SVerticalBox> VerticalHeaderWidget = SNew(SVerticalBox);
+
+	if (HorizontalHeaderWidget->NumSlots() > 0)
+	{
+		VerticalHeaderWidget->AddSlot()
+		.Padding(Padding)
+		.AutoHeight()
+		[
+			HorizontalHeaderWidget
+		];
+	}
+
+	// Custom search if allowed
+	if (HeaderBuilder.GetSearchAllowed())
+	{
+		SearchableKeywords = HeaderBuilder.GetSearchKeywords();
+
+		// Show search box only for root item, not per item
+		if (!CustomizeItem.IsValid())
+		{
+			// Add search box
+			VerticalHeaderWidget->AddSlot()
+            .Padding(Padding)
+            .AutoHeight()
+            [
+            	SAssignNew(SearchBox, SSearchBox)
+            	.HintText(LOCTEXT("OperatorStackEditorStackSearchHint", "Search items"))
+            	.OnTextChanged(this, &SOperatorStackEditorStack::OnSearchTextChanged)
+            	.OnTextCommitted(this, &SOperatorStackEditorStack::OnSearchTextCommitted)
+            ];
+
+			// Add scrollbox with pinned keywords
+			if (!HeaderBuilder.GetSearchPinnedKeywords().IsEmpty())
+			{
+				static const FToolBarStyle& ToolBarStyle = FAppStyle::Get().GetWidgetStyle<FToolBarStyle>("SlimToolBar");
+				static const FCheckBoxStyle* const CheckStyle = &ToolBarStyle.ToggleButton;
+
+				const TSharedPtr<SScrollBox> SearchScroll = SNew(SScrollBox)
+				.Orientation(EOrientation::Orient_Horizontal)
+				.ScrollBarThickness(FVector2D(3.f));
+
+				for (const FString& PinnedKeyword : HeaderBuilder.GetSearchPinnedKeywords())
+				{
+					SearchScroll->AddSlot()
+					.HAlign(HAlign_Center)
+					.VAlign(VAlign_Fill)
+					.Padding(Padding)
+					[
+						SNew(SCheckBox)
+						.Style(CheckStyle)
+						.ForegroundColor(FLinearColor::White)
+						.OnCheckStateChanged(this, &SOperatorStackEditorStack::OnSearchPinnedKeyword, PinnedKeyword)
+						[
+							SNew(STextBlock)
+							.Text(FText::FromString(PinnedKeyword))
+						]
+					];
+				}
+
+				VerticalHeaderWidget->AddSlot()
+				.Padding(Padding)
+				.AutoHeight()
+				[
+					SearchScroll.ToSharedRef()
+				];
+			}
+		}
+	}
+
+	// Custom tool menu on next row
+	const FName& ToolMenuName = HeaderBuilder.GetToolMenuName();
+	const UToolMenu* ToolMenu = !ToolMenuName.IsNone()
+		? UToolMenus::Get()->FindMenu(ToolMenuName)
+		: nullptr;
+
+	if (ToolMenu)
+	{
+		const TSharedRef<SWidget> MenuButtonWidget = SNew(SComboButton)
+			.ToolTipText(HeaderBuilder.GetToolMenuLabel())
+			.OnGetMenuContent(this, &SOperatorStackEditorStack::GenerateMenuWidget, ToolMenuName)
+			.HasDownArrow(false)
+			.ContentPadding(FMargin(Padding))
+			.ButtonContent()
+			[
+				SNew(SHorizontalBox)
+				+ SHorizontalBox::Slot()
+				.AutoWidth()
+				.Padding(Padding)
+				.HAlign(HAlign_Left)
+				.VAlign(VAlign_Center)
+				[
+					SNew(SImage)
+					.DesiredSizeOverride(FVector2D(16.f, 16.f))
+					.Image(HeaderBuilder.GetToolMenuIcon())
+				]
+				+ SHorizontalBox::Slot()
+				.FillWidth(1.f)
+				.Padding(Padding)
+				.HAlign(HAlign_Left)
+				.VAlign(VAlign_Center)
+				[
+					SNew(STextBlock)
+					.Text(HeaderBuilder.GetToolMenuLabel())
+				]
+			];
+
+		VerticalHeaderWidget->AddSlot()
+		.Padding(Padding)
+		.AutoHeight()
+		[
+			MenuButtonWidget
+		];
+	}
+
+	if (CustomizeItem.IsValid())
+	{
+		MessageBoxText = HeaderBuilder.GetMessageBoxText();
+		MessageBoxType = HeaderBuilder.GetMessageBoxType();
+
+		// Custom message box
+		VerticalHeaderWidget->AddSlot()
+		.HAlign(HAlign_Fill)
+		.VAlign(VAlign_Fill)
+		.AutoHeight()
+		[
+			SNew(SOverlay)
+			.Visibility(this, &SOperatorStackEditorStack::GetMessageBoxVisibility)
+			+ SOverlay::Slot()
+			.HAlign(HAlign_Fill)
+			.VAlign(VAlign_Fill)
+			[
+				SNew(SColorBlock)
+				.Color(this, &SOperatorStackEditorStack::GetMessageBoxBackgroundColor)
+			]
+			+ SOverlay::Slot()
+			.HAlign(HAlign_Center)
+			.VAlign(VAlign_Fill)
+			[
+				SNew(SHorizontalBox)
+				+ SHorizontalBox::Slot()
+				.Padding(Padding)
+				.AutoWidth()
+				[
+					SNew(SScaleBox)
+					.Stretch(EStretch::ScaleToFit)
+					.OverrideScreenSize(FVector2D(16.f))
+					[
+						SNew(SImage)
+						.DesiredSizeOverride(FVector2D(16.f))
+						.Visibility(this, &SOperatorStackEditorStack::GetMessageBoxIconVisibility)
+						.Image(this, &SOperatorStackEditorStack::GetMessageBoxIcon)
+					]
+				]
+				+ SHorizontalBox::Slot()
+				.Padding(Padding)
+				.FillWidth(1.f)
+				[
+					SNew(STextBlock)
+					.Justification(ETextJustify::Left)
+					.ColorAndOpacity(FLinearColor::White)
+					.Text(HeaderBuilder.GetMessageBoxText())
+				]
+			]
+		];
+	}
+
+	return SNew(SBorder)
+		.BorderImage(FAppStyle::Get().GetBrush("NoBorder"))
+		.OnMouseButtonDown(this, &SOperatorStackEditorStack::OnHeaderMouseButtonDown)
+		[
+			VerticalHeaderWidget.ToSharedRef()
+		];
+}
+
+TSharedPtr<SWidget> SOperatorStackEditorStack::GenerateBodyWidget()
+{
+	UOperatorStackEditorStackCustomization* StackCustomization = GetStackCustomization();
+	const TSharedPtr<SOperatorStackEditorPanel> MainPanel = GetMainPanel();
+
+	check(MainPanel.IsValid());
+
+	const FOperatorStackEditorTree& ItemTree = MainPanel->GetItemTree(StackCustomization);
+
+	FText EmptyBodyDefaultText = LOCTEXT("EmptyBodyText", "Select a supported item to display it here");
+
+	// We are the root and we have multiple supported items selected
+	if (!CustomizeItem.IsValid() && !Items.IsEmpty())
+	{
+		FOperatorStackEditorBodyBuilder StackBodyBuilder;
+		StackCustomization->CustomizeStackBody(ItemTree, StackBodyBuilder);
+
+		// Gather empty body text
+		if (!StackBodyBuilder.GetEmptyBodyText().IsEmpty())
+		{
+			EmptyBodyDefaultText = StackBodyBuilder.GetEmptyBodyText();
+		}
+
+		// Check if we have a custom widget
+		if (TSharedPtr<SWidget> StackWidget = StackBodyBuilder.GetCustomWidget())
+		{
+			return StackWidget;
+		}
+
+		const TSharedRef<SSplitter> MultiStackBox = SNew(SSplitter)
+			.Orientation(Orient_Vertical);
+
+		for (FOperatorStackEditorItemPtr Item : Items)
+		{
+			TSharedPtr<SOperatorStackEditorStack> ItemStackWidget;
+
+			MultiStackBox->AddSlot()
+			.Value(1.f)
+			[
+				// Scrollbar to be able to scroll in the full view per vertical stack
+				SNew(SScrollBox)
+				.ScrollBarPadding(FMargin(0.f))
+				.ScrollBarAlwaysVisible(true)
+				.ScrollBarThickness(FVector2D(3.f))
+				.ScrollBarVisibility(EVisibility::Hidden)
+				.Orientation(Orient_Vertical)
+				+ SScrollBox::Slot()
+				.FillSize(1.f)
+				.Padding(Padding)
+				[
+					SAssignNew(ItemStackWidget, SOperatorStackEditorStack, GetMainPanel(), StackCustomization, Item)
+				]
+			];
+
+			ItemsWidgets.Add(ItemStackWidget);
+		}
+
+		return MultiStackBox;
+	}
+
+	// We are not the root but we contain children then add a list view
+	if (CustomizeItem.IsValid() && !Items.IsEmpty())
+	{
+		ItemsListView = SNew(SListView<FOperatorStackEditorItemPtr>)
+			.ListViewStyle(&FOperatorStackEditorStyle::Get().GetWidgetStyle<FTableViewStyle>("ListViewStyle"))
+			.ListItemsSource(&Items)
+			.ClearSelectionOnClick(true)
+			.SelectionMode(ESelectionMode::Multi)
+			.OnKeyDownHandler(this, &SOperatorStackEditorStack::OnKeyDownHandler)
+			.OnSelectionChanged(this, &SOperatorStackEditorStack::OnSelectionChanged)
+			.OnGenerateRow(this, &SOperatorStackEditorStack::OnGenerateRow);
+
+		return SNew(SBox)
+			.HAlign(HAlign_Fill)
+			.VAlign(VAlign_Fill)
+			.Padding(0.f)
+			[
+				ItemsListView.ToSharedRef()
+			];
+	}
+
+	// we are a leaf, we have a parent but no children
+	if (CustomizeItem.IsValid() && Items.IsEmpty())
+	{
+		FOperatorStackEditorBodyBuilder ItemBodyBuilder;
+		StackCustomization->CustomizeItemBody(CustomizeItem, ItemTree, ItemBodyBuilder);
+
+		// Gather empty body text
+		if (!ItemBodyBuilder.GetEmptyBodyText().IsEmpty())
+		{
+			EmptyBodyDefaultText = ItemBodyBuilder.GetEmptyBodyText();
+		}
+
+		// We have set a custom widget for this body
+		if (TSharedPtr<SWidget> CustomWidget = ItemBodyBuilder.GetCustomWidget())
+		{
+			return CustomWidget;
+		}
+
+		// Only build details view if we allow it
+		if (ItemBodyBuilder.GetShowDetailsView())
+		{
+			FCustomDetailsViewArgs BodyCustomViewArgs;
+			BodyCustomViewArgs.bShowCategories = false;
+			BodyCustomViewArgs.KeyframeHandler = MainPanel->GetKeyframeHandler();
+			BodyCustomViewArgs.bAllowGlobalExtensions = true;
+			BodyCustomViewArgs.ColumnSizeData = MainPanel->GetDetailColumnSize();
+
+			for (const TSharedPtr<FCustomDetailsViewItemId>& DetailsViewId : ItemBodyBuilder.GetDisallowedDetailsViewItems())
+			{
+				BodyCustomViewArgs.ItemAllowList.Disallow(*DetailsViewId.Get());
+			}
+
+			for (const TSharedPtr<FCustomDetailsViewItemId>& DetailsViewId : ItemBodyBuilder.GetAllowedDetailsViewItems())
+			{
+				BodyCustomViewArgs.ItemAllowList.Allow(*DetailsViewId.Get());
+			}
+
+			const FOperatorStackEditorItem* DetailViewItem = ItemBodyBuilder.GetDetailsViewItem().IsValid()
+				? ItemBodyBuilder.GetDetailsViewItem().Get()
+				: CustomizeItem.Get();
+
+			BodyDetailsView = CreateDetailsView(BodyCustomViewArgs, *DetailViewItem);
+
+			return BodyDetailsView;
+		}
+	}
+
+	// No children and no items, display empty body info text
+	return SNew(SBox)
+		.HAlign(HAlign_Center)
+		.VAlign(VAlign_Center)
+		[
+			SNew(STextBlock)
+			.LineHeightPercentage(3.f)
+			.Text(EmptyBodyDefaultText)
+		];
+}
+
+TSharedPtr<SWidget> SOperatorStackEditorStack::GenerateFooterWidget()
+{
+	UOperatorStackEditorStackCustomization* StackCustomization = GetStackCustomization();
+	const TSharedPtr<SOperatorStackEditorPanel> MainPanel = GetMainPanel();
+
+	if (!CustomizeItem.IsValid())
+	{
+		return nullptr;
+	}
+
+	check(MainPanel.IsValid());
+
+	const FOperatorStackEditorTree& ItemTree = MainPanel->GetItemTree(StackCustomization);
+
+	FOperatorStackEditorFooterBuilder FooterBuilder;
+	StackCustomization->CustomizeItemFooter(CustomizeItem, ItemTree, FooterBuilder);
+
+	// We have set a custom widget for this footer
+	if (TSharedPtr<SWidget> CustomWidget = FooterBuilder.GetCustomWidget())
+	{
+		return CustomWidget;
+	}
+
+	// Only build details view if we allow it
+	if (FooterBuilder.GetShowDetailsView())
+	{
+		FCustomDetailsViewArgs FooterCustomViewArgs;
+		FooterCustomViewArgs.bShowCategories = false;
+		FooterCustomViewArgs.KeyframeHandler = MainPanel->GetKeyframeHandler();
+		FooterCustomViewArgs.bAllowGlobalExtensions = true;
+		FooterCustomViewArgs.ColumnSizeData = MainPanel->GetDetailColumnSize();
+
+		for (const TSharedPtr<FCustomDetailsViewItemId>& DetailsViewId : FooterBuilder.GetDisallowedDetailsViewItems())
+		{
+			FooterCustomViewArgs.ItemAllowList.Disallow(*DetailsViewId.Get());
+		}
+
+		for (const TSharedPtr<FCustomDetailsViewItemId>& DetailsViewId : FooterBuilder.GetAllowedDetailsViewItems())
+		{
+			FooterCustomViewArgs.ItemAllowList.Allow(*DetailsViewId.Get());
+		}
+
+		const FOperatorStackEditorItem* DetailViewItem = FooterBuilder.GetDetailsViewItem().IsValid()
+			? FooterBuilder.GetDetailsViewItem().Get()
+			: CustomizeItem.Get();
+
+		FooterDetailsView = CreateDetailsView(FooterCustomViewArgs, *DetailViewItem);
+
+		return FooterDetailsView;
+	}
+
+	return nullptr;
+}
+
+TSharedRef<SWidget> SOperatorStackEditorStack::GenerateMenuWidget(FName InMenuName) const
+{
+	UOperatorStackEditorMenuContext* const MenuContext = NewObject<UOperatorStackEditorMenuContext>();
+
+	MenuContext->SetContext(GetContext());
+	MenuContext->SetItem(CustomizeItem);
+
+	const TSharedPtr<FExtender> MenuExtender;
+	const FToolMenuContext ToolMenuContext(CommandList, MenuExtender, MenuContext);
+
+	return UToolMenus::Get()->GenerateWidget(InMenuName, ToolMenuContext);
+}
+
+EVisibility SOperatorStackEditorStack::GetBodyVisibility() const
+{
+	return bHeaderExpanded ? EVisibility::Visible : EVisibility::Collapsed;
+}
+
+EVisibility SOperatorStackEditorStack::GetFooterVisibility() const
+{
+	return bHeaderExpanded ? EVisibility::Visible : EVisibility::Collapsed;
+}
+
+EVisibility SOperatorStackEditorStack::GetMessageBoxVisibility() const
+{
+	const EOperatorStackEditorMessageType MessageType = MessageBoxType.Get(EOperatorStackEditorMessageType::None);
+
+	if (MessageType == EOperatorStackEditorMessageType::None)
+	{
+		return EVisibility::Collapsed;
+	}
+
+	const FText MessageText = MessageBoxText.Get(FText::GetEmpty());
+
+	if (MessageText.IsEmpty())
+	{
+		return EVisibility::Collapsed;
+	}
+
+	return EVisibility::Visible;
+}
+
+FLinearColor SOperatorStackEditorStack::GetMessageBoxBackgroundColor() const
+{
+	return MessageBoxColors.FindRef(MessageBoxType.Get(EOperatorStackEditorMessageType::None));
+}
+
+const FSlateBrush* SOperatorStackEditorStack::GetMessageBoxIcon() const
+{
+	return MessageBoxIcons.FindRef(MessageBoxType.Get(EOperatorStackEditorMessageType::None));
+}
+
+EVisibility SOperatorStackEditorStack::GetMessageBoxIconVisibility() const
+{
+	return GetMessageBoxIcon() ? EVisibility::Visible : EVisibility::Collapsed;
+}
+
+void SOperatorStackEditorStack::OnSearchTextChanged(const FText& InSearchText)
+{
+	OnSearchChanged();
+}
+
+void SOperatorStackEditorStack::OnSearchChanged()
+{
+	const TSet<FString> SearchOR = SearchedKeywords;
+	TSet<FString> SearchAND;
+
+	const FString FilterString = SearchBox->GetText().ToString();
+	if (!FilterString.IsEmpty())
+	{
+		SearchAND.Add(FilterString);
+	}
+
+	HandleSearch(SearchOR, SearchAND);
+}
+
+void SOperatorStackEditorStack::OnSearchTextCommitted(const FText& InFilterText, ETextCommit::Type InCommitType)
+{
+	if (InCommitType == ETextCommit::OnCleared)
+	{
+		SearchBox->SetText(FText::GetEmpty());
+		OnSearchChanged();
+		FSlateApplication::Get().ClearKeyboardFocus(EFocusCause::Cleared);
+	}
+}
+
+void SOperatorStackEditorStack::OnSearchPinnedKeyword(ECheckBoxState InCheckState, FString InPinnedKeyword)
+{
+	if (InCheckState == ECheckBoxState::Checked)
+	{
+		SearchedKeywords.Add(InPinnedKeyword);
+	}
+	else if (InCheckState == ECheckBoxState::Unchecked)
+	{
+		SearchedKeywords.Remove(InPinnedKeyword);
+	}
+
+	OnSearchChanged();
+}
+
+bool SOperatorStackEditorStack::MatchSearch(const TSet<FString>& InSearchedKeywordsOR, const TSet<FString>& InSearchedKeywordsAND) const
+{
+	if (InSearchedKeywordsOR.IsEmpty() && InSearchedKeywordsAND.IsEmpty())
+	{
+		return true;
+	}
+
+	// Search in set for full keyword
+	{
+		bool bSearchMatched = false;
+		if (!InSearchedKeywordsOR.IsEmpty())
+		{
+			bSearchMatched = !SearchableKeywords.Intersect(InSearchedKeywordsOR).IsEmpty();
+		}
+
+		if (!InSearchedKeywordsAND.IsEmpty())
+		{
+			bSearchMatched = SearchableKeywords.Includes(InSearchedKeywordsAND);
+		}
+
+		if (bSearchMatched)
+		{
+			return true;
+		}
+	}
+
+	// Search by character for OR keywords
+	bool bORSearchMatched = false;
+	for (const FString& SearchedKeyword : InSearchedKeywordsOR)
+	{
+		for (const FString& Keyword : SearchableKeywords)
+		{
+			if (Keyword.Find(SearchedKeyword, ESearchCase::IgnoreCase, ESearchDir::FromStart) != INDEX_NONE)
+			{
+				bORSearchMatched = true;
+				break;
+			}
+		}
+
+		if (bORSearchMatched)
+		{
+			break;
+		}
+	}
+
+	if (InSearchedKeywordsOR.IsEmpty())
+	{
+		bORSearchMatched = true;
+	}
+
+	// Search by character for AND keywords
+	bool bANDSearchMatched = false;
+	for (const FString& SearchedKeyword : InSearchedKeywordsAND)
+	{
+		for (const FString& Keyword : SearchableKeywords)
+		{
+			if (Keyword.Find(SearchedKeyword, ESearchCase::IgnoreCase, ESearchDir::FromStart) != INDEX_NONE)
+			{
+				bANDSearchMatched = true;
+				break;
+			}
+		}
+
+		if (!bANDSearchMatched)
+		{
+			break;
+		}
+	}
+
+	if (InSearchedKeywordsAND.IsEmpty())
+	{
+		bANDSearchMatched = true;
+	}
+
+	return bORSearchMatched && bANDSearchMatched;
+}
+
+bool SOperatorStackEditorStack::HandleSearch(const TSet<FString>& InSearchedKeywordsOR, const TSet<FString>& InSearchedKeywordsAND)
+{
+	bool bMatchSearch = false;
+
+	for (const TSharedPtr<SOperatorStackEditorStack>& ItemWidget : ItemsWidgets)
+	{
+		if (ItemWidget.IsValid())
+		{
+			if (ItemWidget->HandleSearch(InSearchedKeywordsOR, InSearchedKeywordsAND))
+			{
+				bMatchSearch = true;
+			}
+		}
+	}
+
+	// Search for properties in details view
+	if (BodyDetailsView)
+	{
+		bMatchSearch |= BodyDetailsView->FilterItems(InSearchedKeywordsAND.Array());
+	}
+
+	if (FooterDetailsView)
+	{
+		bMatchSearch |= FooterDetailsView->FilterItems(InSearchedKeywordsAND.Array());
+	}
+
+	// Do not hide root item if nothing was found
+	if (CustomizeItem.IsValid())
+	{
+		if (bMatchSearch || MatchSearch(InSearchedKeywordsOR, InSearchedKeywordsAND))
+		{
+			bMatchSearch = true;
+			SetVisibility(EVisibility::Visible);
+		}
+		else
+		{
+			SetVisibility(EVisibility::Collapsed);
+		}
+	}
+
+	return bMatchSearch;
+}
+
+TSharedRef<SWidget> SOperatorStackEditorStack::GenerateStackWidget()
+{
+	const TSharedRef<SVerticalBox> StackWidget = SNew(SVerticalBox);
+
+	const TSharedPtr<SWidget> HeaderWidget = GenerateHeaderWidget();
+	const TSharedPtr<SWidget> BodyWidget = GenerateBodyWidget();
+	const TSharedPtr<SWidget> FooterWidget = GenerateFooterWidget();
+
+	if (HeaderWidget.IsValid())
+	{
+		StackWidget->AddSlot()
+		.AutoHeight()
+		[
+			HeaderWidget.ToSharedRef()
+		];
+	}
+
+	if (BodyWidget.IsValid())
+	{
+		StackWidget->AddSlot()
+		.FillHeight(1.f)
+		[
+			SNew(SBox)
+			.Padding(0.f)
+			.Visibility(this, &SOperatorStackEditorStack::GetBodyVisibility)
+			[
+				BodyWidget.ToSharedRef()
+			]
+		];
+	}
+
+	if (FooterWidget.IsValid())
+	{
+		StackWidget->AddSlot()
+		.AutoHeight()
+		[
+			SNew(SBox)
+			.Visibility(this, &SOperatorStackEditorStack::GetFooterVisibility)
+			[
+				FooterWidget.ToSharedRef()
+			]
+		];
+	}
+
+	// Only add border if we have a parent
+	if (CustomizeItem.IsValid() && Items.IsEmpty())
+	{
+		return SNew(SHorizontalBox)
+		+ SHorizontalBox::Slot()
+		.Padding(0.f)
+		.AutoWidth()
+		[
+			SNew(SSeparator)
+			.Visibility(BorderColor == FLinearColor::Transparent ? EVisibility::Collapsed : EVisibility::Visible)
+			.ColorAndOpacity(BorderColor)
+			.SeparatorImage(FAppStyle::GetBrush("ThinLine.Horizontal"))
+			.Thickness(3.0f)
+			.Orientation(EOrientation::Orient_Vertical)
+		]
+		+ SHorizontalBox::Slot()
+		.Padding(0.f)
+		.FillWidth(1.f)
+		[
+			SNew(SBox)
+			.Padding(2.f)
+			[
+				SNew(SOverlay)
+				+ SOverlay::Slot()
+				[
+					SNew(SColorBlock)
+					.Color(FOperatorStackEditorStyle::Get().GetColor("ForegroundColor"))
+				]
+				+ SOverlay::Slot()
+				[
+					StackWidget
+				]
+			]
+		];
+	}
+
+	return StackWidget;
+}
+
+void SOperatorStackEditorStack::OnHeaderExpansionChanged(bool bInExpansion)
+{
+	bHeaderExpanded = bInExpansion;
+
+	if (const TSharedPtr<SOperatorStackEditorPanel> MainPanel = GetMainPanel())
+	{
+		if (CustomizeItem.IsValid())
+		{
+			MainPanel->SaveItemExpansionState(CustomizeItem->GetValuePtr(), bHeaderExpanded);
+		}
+	}
+}
+
+void SOperatorStackEditorStack::OnSelectionChanged(FOperatorStackEditorItemPtr InItem, ESelectInfo::Type InSelect) const
+{
+	if (InSelect != ESelectInfo::Direct && !IsSelectableRow(InItem))
+	{
+		ItemsListView->SetItemSelection(InItem, false);
+	}
+}
+
+FReply SOperatorStackEditorStack::OnKeyDownHandler(const FGeometry& InGeometry, const FKeyEvent& InKeyEvent)
+{
+	bool bProcessed = false;
+
+	if (ItemsListView.IsValid() && ItemsListView->GetNumItemsSelected() > 0)
+	{
+		for (FOperatorStackEditorItemPtr Item : ItemsListView->GetSelectedItems())
+		{
+			const int32 Idx = Items.Find(Item);
+
+			if (Idx != INDEX_NONE && ItemsWidgets.IsValidIndex(Idx))
+			{
+				if (ItemsWidgets[Idx]->CommandList.IsValid())
+				{
+					bProcessed = ItemsWidgets[Idx]->CommandList->ProcessCommandBindings(InKeyEvent);
+				}
+			}
+		}
+	}
+	else if (CustomizeItem.IsValid() && CommandList.IsValid())
+	{
+		bProcessed = CommandList->ProcessCommandBindings(InKeyEvent);
+	}
+
+	if (bProcessed)
+	{
+		return FReply::Handled();
+	}
+
+	return FReply::Unhandled();
+}
+
+FReply SOperatorStackEditorStack::OnHeaderMouseButtonDown(const FGeometry& MyGeometry, const FPointerEvent& MouseEvent)
+{
+	// Open context menu
+	if (CustomizeItem.IsValid() && MouseEvent.GetEffectingButton() == EKeys::RightMouseButton)
+	{
+		if (UToolMenus::Get()->IsMenuRegistered(ContextMenuName))
+		{
+			FSlateApplication::Get().PushMenu(
+				AsShared(),
+				FWidgetPath(),
+				GenerateMenuWidget(ContextMenuName),
+				MouseEvent.GetLastScreenSpacePosition(),
+				FPopupTransitionEffect(FPopupTransitionEffect::ContextMenu)
+			);
+
+			return FReply::Handled();
+		}
+	}
+
+	return FReply::Unhandled();
+}
+
+bool SOperatorStackEditorStack::IsSelectableRow(FOperatorStackEditorItemPtr InItem) const
+{
+	UOperatorStackEditorStackCustomization* StackCustomization = GetStackCustomization();
+
+	if (InItem.IsValid() && StackCustomization)
+	{
+		return StackCustomization->OnIsItemDraggable(InItem);
+	}
+
+	return false;
+}
+
+TSharedRef<ICustomDetailsView> SOperatorStackEditorStack::CreateDetailsView(const FCustomDetailsViewArgs& InArgs, const FOperatorStackEditorItem& InItem)
+{
+	TSharedRef<ICustomDetailsView> CustomDetailsView = ICustomDetailsViewModule::Get().CreateCustomDetailsView(InArgs);
+
+	if (InItem.GetValueType().GetTypeEnum() == EOperatorStackEditorItemType::Object)
+	{
+		UObject* Item = InItem.Get<UObject>();
+		CustomDetailsView->SetObject(Item);
+	}
+	else if (InItem.GetValueType().GetTypeEnum() == EOperatorStackEditorItemType::Struct)
+	{
+		const FOperatorStackEditorStructItem* Item = static_cast<const FOperatorStackEditorStructItem*>(&InItem);
+		CustomDetailsView->SetStruct(Item->GetStructOnScope());
+	}
+
+	return CustomDetailsView;
+}
+
+#undef LOCTEXT_NAMESPACE
diff --git a/Engine/Plugins/Experimental/OperatorStack/Source/OperatorStack/Private/Widgets/SOperatorStackEditorStack.h b/Engine/Plugins/Experimental/OperatorStack/Source/OperatorStack/Private/Widgets/SOperatorStackEditorStack.h
new file mode 100644
--- /dev/null
+++ b/Engine/Plugins/Experimental/OperatorStack/Source/OperatorStack/Private/Widgets/SOperatorStackEditorStack.h
@@ -0,0 +1,154 @@
+// Copyright Epic Games, Inc. All Rights Reserved.
+
+#pragma once
+
+#include "Builders/OperatorStackEditorHeaderBuilder.h"
+#include "Contexts/OperatorStackEditorContext.h"
+#include "Items/OperatorStackEditorItem.h"
+#include "Templates/SharedPointer.h"
+#include "Widgets/SCompoundWidget.h"
+#include "Widgets/Views/SListView.h"
+
+enum class ECheckBoxState : uint8;
+class FReply;
+class FUICommandList;
+class ICustomDetailsView;
+class SOperatorStackEditorPanel;
+class SSearchBox;
+class UOperatorStackEditorStackCustomization;
+struct EVisibility;
+struct FButtonStyle;
+struct FCustomDetailsViewArgs;
+
+/** Represent a stack customization widget */
+class SOperatorStackEditorStack : public SCompoundWidget
+{
+	friend class SOperatorStackEditorStackRow;
+
+public:
+	static inline constexpr float Padding = 3.f;
+	static const TMap<EOperatorStackEditorMessageType, FLinearColor> MessageBoxColors;
+	static const TMap<EOperatorStackEditorMessageType, const FSlateBrush*> MessageBoxIcons;
+
+	SLATE_BEGIN_ARGS(SOperatorStackEditorStack) {}
+	SLATE_END_ARGS()
+
+	void Construct(const FArguments& InArgs
+		, const TSharedPtr<SOperatorStackEditorPanel> InMainPanel
+		, TObjectPtr<UOperatorStackEditorStackCustomization> InCustomization
+		, const FOperatorStackEditorItemPtr& InCustomizeItem);
+
+	UOperatorStackEditorStackCustomization* GetStackCustomization() const
+	{
+		return StackCustomizationWeak.Get();
+	}
+
+	TSharedPtr<SOperatorStackEditorPanel> GetMainPanel() const
+	{
+		return MainPanelWeak.Pin();
+	}
+
+	FOperatorStackEditorContextPtr GetContext() const;
+
+protected:
+	TSharedRef<ITableRow> OnGenerateRow(FOperatorStackEditorItemPtr InItem, const TSharedRef<STableViewBase>& InOwnerTable);
+
+	/** Key down handler for list */
+	FReply OnKeyDownHandler(const FGeometry& InGeometry, const FKeyEvent& InKeyEvent);
+
+	/** Mouse event for header */
+	FReply OnHeaderMouseButtonDown(const FGeometry& MyGeometry, const FPointerEvent& MouseEvent);
+
+	/** Header expansion state changed */
+	void OnHeaderExpansionChanged(bool bInExpansion);
+
+	/** List selection changed */
+	void OnSelectionChanged(FOperatorStackEditorItemPtr InItem, ESelectInfo::Type InSelect) const;
+
+	/** Is this item selectable in the list view, here check if item is draggable */
+	bool IsSelectableRow(FOperatorStackEditorItemPtr InItem) const;
+
+	/** A stack widget consists of a header, body and footer */
+	TSharedRef<SWidget> GenerateStackWidget();
+
+	TSharedPtr<SWidget> GenerateHeaderWidget();
+	TSharedPtr<SWidget> GenerateBodyWidget();
+	TSharedPtr<SWidget> GenerateFooterWidget();
+
+	/** Generate this named menu widget for this context */
+	TSharedRef<SWidget> GenerateMenuWidget(FName InMenuName) const;
+
+	EVisibility GetBodyVisibility() const;
+	EVisibility GetFooterVisibility() const;
+
+	EVisibility GetMessageBoxVisibility() const;
+	FLinearColor GetMessageBoxBackgroundColor() const;
+	const FSlateBrush* GetMessageBoxIcon() const;
+	EVisibility GetMessageBoxIconVisibility() const;
+
+	/** The main panel we are in */
+	TWeakPtr<SOperatorStackEditorPanel> MainPanelWeak;
+
+	/** Customization item for this stack */
+	FOperatorStackEditorItemPtr CustomizeItem;
+
+	/** Children items of this item */
+	TArray<FOperatorStackEditorItemPtr> Items;
+
+	/** The list view if we contain any children */
+	TSharedPtr<SListView<FOperatorStackEditorItemPtr>> ItemsListView;
+
+	/** Header custom view for properties in header */
+	TSharedPtr<ICustomDetailsView> HeaderDetailsView;
+
+	/** Body custom view for properties in body */
+	TSharedPtr<ICustomDetailsView> BodyDetailsView;
+
+	/** Footer custom view for properties in footer */
+	TSharedPtr<ICustomDetailsView> FooterDetailsView;
+
+	/** The customization to use */
+	TWeakObjectPtr<UOperatorStackEditorStackCustomization> StackCustomizationWeak;
+
+	/** Children items widget created from this, used for search propagation */
+	TArray<TSharedPtr<SOperatorStackEditorStack>> ItemsWidgets;
+
+	/** The search bar if any, only for root operator stack */
+	TSharedPtr<SSearchBox> SearchBox;
+
+	/** Contains pinned searched keywords, only for root operator stack */
+	TSet<FString> SearchedKeywords;
+
+	/** Contains searchable keywords for this item */
+	TSet<FString> SearchableKeywords;
+
+	/** Is header expanded to show body and footer */
+	bool bHeaderExpanded = true;
+
+	/** Color of the border around this stack */
+	FLinearColor BorderColor = FLinearColor::Transparent;
+
+	/** List of command available for this item */
+	TSharedPtr<FUICommandList> CommandList;
+
+	/** Name of the context menu */
+	FName ContextMenuName;
+
+	/** When set will display an alert message in the header */
+	TAttribute<FText> MessageBoxText;
+
+	/** When set will display an alert message in the header */
+	TAttribute<EOperatorStackEditorMessageType> MessageBoxType;
+
+private:
+	static TSharedRef<ICustomDetailsView> CreateDetailsView(const FCustomDetailsViewArgs& InArgs, const FOperatorStackEditorItem& InItem);
+
+	void OnSearchTextChanged(const FText& InSearchText);
+	void OnSearchTextCommitted(const FText& InFilterText, ETextCommit::Type InCommitType);
+	void OnSearchPinnedKeyword(ECheckBoxState InCheckState, FString InPinnedKeyword);
+
+	void OnSearchChanged();
+	bool MatchSearch(const TSet<FString>& InSearchedKeywordsOR, const TSet<FString>& InSearchedKeywordsAND) const;
+	bool HandleSearch(const TSet<FString>& InSearchedKeywordsOR, const TSet<FString>& InSearchedKeywordsAND);
+};
+
diff --git a/Engine/Plugins/Experimental/OperatorStack/Source/OperatorStack/Private/Widgets/SOperatorStackExpanderButton.cpp b/Engine/Plugins/Experimental/OperatorStack/Source/OperatorStack/Private/Widgets/SOperatorStackExpanderButton.cpp
new file mode 100644
--- /dev/null
+++ b/Engine/Plugins/Experimental/OperatorStack/Source/OperatorStack/Private/Widgets/SOperatorStackExpanderButton.cpp
@@ -0,0 +1,87 @@
+// Copyright Epic Games, Inc. All Rights Reserved.
+
+#include "SOperatorStackExpanderButton.h"
+#include "Widgets/Images/SImage.h"
+#include "Widgets/Input/SButton.h"
+#include "Widgets/Layout/SBox.h"
+
+void SOperatorStackExpanderButton::Construct(const FArguments& InArgs)
+{
+	bIsItemExpanded = InArgs._StartsExpanded;
+	
+	CollapsedButtonIcon = InArgs._CollapsedButtonIcon;
+	ExpandedButtonIcon = InArgs._ExpandedButtonIcon;
+	CollapsedHoveredButtonIcon = InArgs._CollapsedHoveredButtonIcon;
+	ExpandedHoveredButtonIcon = InArgs._ExpandedHoveredButtonIcon;
+	
+	if (InArgs._OnExpansionStateChanged.IsBound())
+	{
+		OnExpansionStateChangedEvent = InArgs._OnExpansionStateChanged;
+	}
+
+	ChildSlot
+	[
+		SNew(SBox)
+		.MinDesiredWidth(12)
+		.MaxDesiredHeight(12)
+		[
+			SAssignNew(ExpanderButton, SButton)
+			.ButtonStyle(FCoreStyle::Get(), "NoBorder")
+			.VAlign(VAlign_Center)
+			.HAlign(HAlign_Center)
+			.ClickMethod(EButtonClickMethod::MouseDown)
+			.OnClicked(this, &SOperatorStackExpanderButton::OnExpanderButtonClicked)
+			.ContentPadding(0)
+			.IsFocusable(false)
+			[
+				SNew(SImage)
+				.Image(this, &SOperatorStackExpanderButton::GetExpanderIcon)
+				.ColorAndOpacity(FSlateColor::UseSubduedForeground())
+			]
+		]
+	];
+
+	BroadcastExpansion();
+}
+
+const FSlateBrush* SOperatorStackExpanderButton::GetExpanderIcon() const
+{
+	if (bIsItemExpanded)
+	{
+		if (ExpanderButton->IsHovered())
+		{
+			return ExpandedHoveredButtonIcon;
+		}
+		
+		return ExpandedButtonIcon;
+	}
+
+	if (ExpanderButton->IsHovered())
+	{
+		return CollapsedHoveredButtonIcon;
+	}
+
+	return CollapsedButtonIcon;
+}
+
+FReply SOperatorStackExpanderButton::OnExpanderButtonClicked()
+{
+	ToggleExpansion();
+
+	return FReply::Handled();
+}
+
+void SOperatorStackExpanderButton::BroadcastExpansion() const
+{
+	if (OnExpansionStateChangedEvent.IsBound())
+	{
+		OnExpansionStateChangedEvent.Execute(bIsItemExpanded);
+	}
+}
+
+void SOperatorStackExpanderButton::ToggleExpansion()
+{
+	bIsItemExpanded = !bIsItemExpanded;
+
+	BroadcastExpansion();
+}
diff --git a/Engine/Plugins/Experimental/OperatorStack/Source/OperatorStack/Private/Widgets/SOperatorStackExpanderButton.h b/Engine/Plugins/Experimental/OperatorStack/Source/OperatorStack/Private/Widgets/SOperatorStackExpanderButton.h
new file mode 100644
--- /dev/null
+++ b/Engine/Plugins/Experimental/OperatorStack/Source/OperatorStack/Private/Widgets/SOperatorStackExpanderButton.h
@@ -0,0 +1,56 @@
+// Copyright Epic Games, Inc. All Rights Reserved.
+
+#pragma once
+
+#include "Styling/AppStyle.h"
+#include "Styling/SlateBrush.h"
+#include "Templates/SharedPointer.h"
+#include "Widgets/DeclarativeSyntaxSupport.h"
+#include "Widgets/SCompoundWidget.h"
+
+class SButton;
+
+DECLARE_DELEGATE_OneParam(FOperatorStackExpanderButtonStateChanged, bool)
+
+/** A simple expander button with image, fires an event everytime state changes */
+class SOperatorStackExpanderButton : public SCompoundWidget
+{
+public:
+	SLATE_BEGIN_ARGS(SOperatorStackExpanderButton)
+		: _StartsExpanded(true)
+		, _CollapsedButtonIcon(FAppStyle::Get().GetBrush(TEXT("TreeArrow_Collapsed")))
+		, _ExpandedButtonIcon(FAppStyle::Get().GetBrush(TEXT("TreeArrow_Expanded")))
+		, _CollapsedHoveredButtonIcon(FAppStyle::Get().GetBrush(TEXT("TreeArrow_Collapsed_Hovered")))
+		, _ExpandedHoveredButtonIcon(FAppStyle::Get().GetBrush(TEXT("TreeArrow_Expanded_Hovered")))
+	{}
+	/** Are we expanded by default */
+	SLATE_ARGUMENT(bool, StartsExpanded)
+	/** Image to change the look of the button */
+	SLATE_ARGUMENT(const FSlateBrush*, CollapsedButtonIcon)
+	SLATE_ARGUMENT(const FSlateBrush*, ExpandedButtonIcon)
+	SLATE_ARGUMENT(const FSlateBrush*, CollapsedHoveredButtonIcon)
+	SLATE_ARGUMENT(const FSlateBrush*, ExpandedHoveredButtonIcon)
+	/** Called when the expansion state changes */
+	SLATE_EVENT(FOperatorStackExpanderButtonStateChanged, OnExpansionStateChanged)
+	SLATE_END_ARGS()
+	
+	void Construct(const FArguments& InArgs);
+
+private:
+	const FSlateBrush* GetExpanderIcon() const;
+	FReply OnExpanderButtonClicked();
+	
+	void BroadcastExpansion() const;
+	void ToggleExpansion();
+	
+	FOperatorStackExpanderButtonStateChanged OnExpansionStateChangedEvent;
+	
+	TSharedPtr<SButton> ExpanderButton;
+	
+	const FSlateBrush* CollapsedButtonIcon = nullptr;
+	const FSlateBrush* ExpandedButtonIcon = nullptr;
+	const FSlateBrush* CollapsedHoveredButtonIcon = nullptr;
+	const FSlateBrush* ExpandedHoveredButtonIcon = nullptr;
+	
+	bool bIsItemExpanded = true;
+};
diff --git a/Engine/Plugins/Experimental/OperatorStack/Source/OperatorStack/Private/Widgets/TableRows/SOperatorStackEditorStackRow.cpp b/Engine/Plugins/Experimental/OperatorStack/Source/OperatorStack/Private/Widgets/TableRows/SOperatorStackEditorStackRow.cpp
new file mode 100644
--- /dev/null
+++ b/Engine/Plugins/Experimental/OperatorStack/Source/OperatorStack/Private/Widgets/TableRows/SOperatorStackEditorStackRow.cpp
@@ -0,0 +1,151 @@
+// Copyright Epic Games, Inc. All Rights Reserved.
+
+#include "Widgets/TableRows/SOperatorStackEditorStackRow.h"
+
+#include "Customizations/OperatorStackEditorStackCustomization.h"
+#include "Framework/Application/SlateApplication.h"
+#include "Styles/OperatorStackEditorStyle.h"
+#include "Widgets/DragDropOps/OperatorStackEditorDragDropOp.h"
+#include "Widgets/Layout/SSeparator.h"
+
+void SOperatorStackEditorStackRow::Construct(const FArguments& InArgs
+	, const TSharedRef<STableViewBase>& InOwnerTableView
+	, const TSharedPtr<SOperatorStackEditorStack>& InOuterStack
+	, const FOperatorStackEditorItemPtr& InCustomizeItem)
+{
+	OuterStackWeak = InOuterStack;
+
+	check(InOuterStack->GetStackCustomization());
+
+	UOperatorStackEditorStackCustomization* Customization = InOuterStack->GetStackCustomization();
+	TSharedPtr<SOperatorStackEditorPanel> MainPanel = InOuterStack->MainPanelWeak.Pin();
+
+	InnerStack = SNew(SOperatorStackEditorStack, MainPanel, Customization, InCustomizeItem);
+
+	InOuterStack->ItemsWidgets.Add(InnerStack);
+
+	ChildSlot
+	[
+		SNew(SVerticalBox)
+		.Visibility(MakeAttributeLambda([this]()->EVisibility
+		{
+			return InnerStack->GetVisibility();
+		}))
+		+ SVerticalBox::Slot()
+		.AutoHeight()
+		[
+			SNew(SSeparator)
+			.ColorAndOpacity(FOperatorStackEditorStyle::Get().GetColor("BackgroundColor"))
+			.SeparatorImage(FAppStyle::GetBrush("ThinLine.Horizontal"))
+			.Thickness(SOperatorStackEditorStack::Padding)
+			.Orientation(EOrientation::Orient_Horizontal)
+		]
+		+ SVerticalBox::Slot()
+		.FillHeight(1.f)
+		[
+			InnerStack.ToSharedRef()
+		]
+		+ SVerticalBox::Slot()
+		.AutoHeight()
+		[
+			SNew(SSeparator)
+			.ColorAndOpacity(FOperatorStackEditorStyle::Get().GetColor("BackgroundColor"))
+			.SeparatorImage(FAppStyle::GetBrush("ThinLine.Horizontal"))
+			.Thickness(SOperatorStackEditorStack::Padding)
+			.Orientation(EOrientation::Orient_Horizontal)
+		]
+	];
+
+	STableRow<FOperatorStackEditorItemPtr>::ConstructInternal(
+		STableRow<FOperatorStackEditorItemPtr>::FArguments()
+			.ShowSelection(true)
+			.OnDragDetected(this, &SOperatorStackEditorStackRow::OnStackDragDetected)
+			.OnCanAcceptDrop(this, &SOperatorStackEditorStackRow::OnStackCanAcceptDrop)
+			.OnDrop(this, &SOperatorStackEditorStackRow::OnStackDrop),
+			InOwnerTableView
+	);
+
+	SetBorderImage(TAttribute<const FSlateBrush*>(this, &SOperatorStackEditorStackRow::GetBorder));
+}
+
+FReply SOperatorStackEditorStackRow::OnStackDragDetected(const FGeometry& InGeometry, const FPointerEvent& InPointer) const
+{
+	const TSharedPtr<SOperatorStackEditorStack> OuterStack = OuterStackWeak.Pin();
+
+	if (OuterStack.IsValid() && OuterStack->ItemsListView.IsValid())
+	{
+		const TArray<TSharedPtr<FOperatorStackEditorItem>>& SelectedItems = OuterStack->ItemsListView->GetSelectedItems();
+
+		if (!SelectedItems.IsEmpty())
+		{
+			const TSharedRef<FOperatorStackEditorDragDropOp> StackDragDropOp = FOperatorStackEditorDragDropOp::New(SelectedItems);
+
+			return FReply::Handled().BeginDragDrop(StackDragDropOp);
+		}
+	}
+
+	return FReply::Unhandled();
+}
+
+TOptional<EItemDropZone> SOperatorStackEditorStackRow::OnStackCanAcceptDrop(const FDragDropEvent& InDropEvent
+	, EItemDropZone InZone
+	, FOperatorStackEditorItemPtr InZoneItem) const
+{
+	if (const TSharedPtr<FOperatorStackEditorDragDropOp> StackDragDropOp = InDropEvent.GetOperationAs<FOperatorStackEditorDragDropOp>())
+	{
+		if (InnerStack.IsValid())
+		{
+			UOperatorStackEditorStackCustomization* StackCustomization = InnerStack->GetStackCustomization();
+
+			TOptional<EItemDropZone> DropZone = StackCustomization->OnItemCanAcceptDrop(StackDragDropOp->GetDraggedItems(), InZoneItem, InZone);
+			StackDragDropOp->SetDropZone(DropZone);
+
+			return DropZone;
+		}
+	}
+
+	return TOptional<EItemDropZone>();
+}
+
+FReply SOperatorStackEditorStackRow::OnStackDrop(FDragDropEvent const& InEvent) const
+{
+	if (const TSharedPtr<FOperatorStackEditorDragDropOp> StackDragDropOp = InEvent.GetOperationAs<FOperatorStackEditorDragDropOp>())
+	{
+		const TOptional<EItemDropZone> DropZone = StackDragDropOp->GetDropZone();
+
+		if (InnerStack.IsValid() && DropZone.IsSet())
+		{
+			UOperatorStackEditorStackCustomization* StackCustomization = InnerStack->GetStackCustomization();
+			const FOperatorStackEditorItemPtr ZoneItem = InnerStack->CustomizeItem;
+
+			StackCustomization->OnDropItem(StackDragDropOp->GetDraggedItems(), ZoneItem, DropZone.GetValue());
+		}
+	}
+
+	return FReply::Handled().EndDragDrop();
+}
+
+int32 SOperatorStackEditorStackRow::OnPaintDropIndicator(EItemDropZone InZone, const FPaintArgs& InPaintArgs, const FGeometry& InGeometry, const FSlateRect& InSlateRect, FSlateWindowElementList& OutDrawElements, int32 InLayerId, const FWidgetStyle& InWidgetStyle, bool bInParentEnabled) const
+{
+	if (const TSharedPtr<FDragDropOperation> DragDropOperation = FSlateApplication::Get().GetDragDroppingContent())
+	{
+		if (DragDropOperation->IsOfType<FOperatorStackEditorDragDropOp>())
+		{
+			const FSlateBrush* DropIndicatorBrush = GetDropIndicatorBrush(InZone);
+
+			constexpr float OffsetX = 10.0f;
+			const FVector2D Offset(OffsetX * GetIndentLevel(), 0.f);
+			FSlateDrawElement::MakeBox
+			(
+				OutDrawElements,
+				InLayerId++,
+				InGeometry.ToPaintGeometry(FVector2D(InGeometry.GetLocalSize() - Offset), FSlateLayoutTransform(Offset)),
+				DropIndicatorBrush,
+				ESlateDrawEffect::None,
+				DropIndicatorBrush->GetTint(InWidgetStyle) * InWidgetStyle.GetColorAndOpacityTint()
+			);
+		}
+	}
+
+	return InLayerId;
+}
diff --git a/Engine/Plugins/Experimental/OperatorStack/Source/OperatorStack/Private/Widgets/TableRows/SOperatorStackEditorStackRow.h b/Engine/Plugins/Experimental/OperatorStack/Source/OperatorStack/Private/Widgets/TableRows/SOperatorStackEditorStackRow.h
new file mode 100644
--- /dev/null
+++ b/Engine/Plugins/Experimental/OperatorStack/Source/OperatorStack/Private/Widgets/TableRows/SOperatorStackEditorStackRow.h
@@ -0,0 +1,44 @@
+// Copyright Epic Games, Inc. All Rights Reserved.
+
+#pragma once
+
+#include "Items/OperatorStackEditorItem.h"
+#include "Templates/SharedPointer.h"
+#include "Widgets/SOperatorStackEditorStack.h"
+#include "Widgets/Views/STableRow.h"
+#include "Widgets/Views/STreeView.h"
+
+/** A single row item in the stack listview that contains a stack, get it it's recursive :) */
+class SOperatorStackEditorStackRow : public STableRow<FOperatorStackEditorItemPtr>
+{
+
+public:
+	SLATE_BEGIN_ARGS(SOperatorStackEditorStackRow) {}
+	SLATE_END_ARGS()
+
+	void Construct(const FArguments& InArgs
+		, const TSharedRef<STableViewBase>& InOwnerTableView
+		, const TSharedPtr<SOperatorStackEditorStack>& InOuterStack
+		, const FOperatorStackEditorItemPtr& InCustomizeItem);
+
+	TSharedPtr<SOperatorStackEditorStack> GetInnerStack() const
+	{
+		return InnerStack;
+	}
+
+protected:
+	/** When starting a stack drag operation */
+	FReply OnStackDragDetected(const FGeometry& InGeometry, const FPointerEvent& InPointer) const;
+
+	/** Whether we allow or disallow drop */
+	TOptional<EItemDropZone> OnStackCanAcceptDrop(const FDragDropEvent& InDropEvent, EItemDropZone InZone, FOperatorStackEditorItemPtr InZoneItem) const;
+
+	/** When a drop event happens on an item */
+	FReply OnStackDrop(FDragDropEvent const& InEvent) const;
+
+	/** Draw indicator for drop */
+	virtual int32 OnPaintDropIndicator(EItemDropZone InZone, const FPaintArgs& InPaintArgs, const FGeometry& InGeometry, const FSlateRect& InSlateRect, FSlateWindowElementList& OutDrawElements, int32 InLayerId, const FWidgetStyle& InWidgetStyle, bool bInParentEnabled) const override;
+
+	TWeakPtr<SOperatorStackEditorStack> OuterStackWeak;
+	TSharedPtr<SOperatorStackEditorStack> InnerStack;
+};
diff --git a/Engine/Plugins/Experimental/OperatorStack/Source/OperatorStack/Public/Builders/OperatorStackEditorBodyBuilder.h b/Engine/Plugins/Experimental/OperatorStack/Source/OperatorStack/Public/Builders/OperatorStackEditorBodyBuilder.h
new file mode 100644
--- /dev/null
+++ b/Engine/Plugins/Experimental/OperatorStack/Source/OperatorStack/Public/Builders/OperatorStackEditorBodyBuilder.h
@@ -0,0 +1,86 @@
+// Copyright Epic Games, Inc. All Rights Reserved.
+
+#pragma once
+
+#include "Internationalization/Text.h"
+#include "Templates/SharedPointer.h"
+
+class FCustomDetailsViewItemId;
+class SWidget;
+struct FOperatorStackEditorItem;
+
+/** Body builder for item */
+struct OPERATORSTACKEDITOR_API FOperatorStackEditorBodyBuilder
+{
+	/** Override the stack generated widget with a custom one */
+	FOperatorStackEditorBodyBuilder& SetCustomWidget(TSharedPtr<SWidget> InWidget);
+	
+	/** Show a details view for the current item */
+	FOperatorStackEditorBodyBuilder& SetShowDetailsView(bool bInShowDetailsView);
+
+	FOperatorStackEditorBodyBuilder& SetDetailsViewItem(TSharedPtr<FOperatorStackEditorItem> InItem);
+	
+	/** Disallow specific property in the details view */
+	FOperatorStackEditorBodyBuilder& DisallowProperty(FProperty* InProperty);
+
+	/** Disallow specific category in the details view */
+	FOperatorStackEditorBodyBuilder& DisallowCategory(const FName& InCategory);
+
+	/** Allow specific property in the details view */
+	FOperatorStackEditorBodyBuilder& AllowProperty(FProperty* InProperty);
+
+	/** Allow specific category in the details view */
+	FOperatorStackEditorBodyBuilder& AllowCategory(const FName& InCategory);
+
+	/** Set a text that will be displayed when it is empty */
+	FOperatorStackEditorBodyBuilder& SetEmptyBodyText(const FText& InText);
+
+	TSharedPtr<SWidget> GetCustomWidget() const
+	{
+		return CustomWidget;
+	}
+
+	bool GetShowDetailsView() const
+	{
+		return bShowDetailsView;
+	}
+	
+	TSharedPtr<FOperatorStackEditorItem> GetDetailsViewItem() const
+	{
+		return DetailsViewItem;	
+	}
+	
+	const TArray<TSharedPtr<FCustomDetailsViewItemId>>& GetDisallowedDetailsViewItems() const
+	{
+		return DisallowedDetailsViewItems;
+	}
+
+	const TArray<TSharedPtr<FCustomDetailsViewItemId>>& GetAllowedDetailsViewItems() const
+	{
+		return AllowedDetailsViewItems;
+	}
+	
+	const FText& GetEmptyBodyText() const
+	{
+		return EmptyBodyText;
+	}
+
+protected:
+	/** Custom widget to replace content */
+	TSharedPtr<SWidget> CustomWidget = nullptr;
+	
+	/** Does this body contains a details view */
+	bool bShowDetailsView = false;
+	
+	/** Override actual item to display detail view */
+	TSharedPtr<FOperatorStackEditorItem> DetailsViewItem = nullptr;
+
+	/** When the body is empty, this will be displayed */
+	FText EmptyBodyText;
+
+	/** Disallowed items inside details view */
+	TArray<TSharedPtr<FCustomDetailsViewItemId>> DisallowedDetailsViewItems;
+	
+	/** Allowed items inside details view */
+	TArray<TSharedPtr<FCustomDetailsViewItemId>> AllowedDetailsViewItems;
+};
\ No newline at end of file
diff --git a/Engine/Plugins/Experimental/OperatorStack/Source/OperatorStack/Public/Builders/OperatorStackEditorFooterBuilder.h b/Engine/Plugins/Experimental/OperatorStack/Source/OperatorStack/Public/Builders/OperatorStackEditorFooterBuilder.h
new file mode 100644
--- /dev/null
+++ b/Engine/Plugins/Experimental/OperatorStack/Source/OperatorStack/Public/Builders/OperatorStackEditorFooterBuilder.h
@@ -0,0 +1,8 @@
+// Copyright Epic Games, Inc. All Rights Reserved.
+
+#pragma once
+
+#include "OperatorStackEditorBodyBuilder.h"
+
+/** Body and Footer are the same at the moment, we may add additional behaviour here */
+typedef FOperatorStackEditorBodyBuilder FOperatorStackEditorFooterBuilder;
\ No newline at end of file
diff --git a/Engine/Plugins/Experimental/OperatorStack/Source/OperatorStack/Public/Builders/OperatorStackEditorHeaderBuilder.h b/Engine/Plugins/Experimental/OperatorStack/Source/OperatorStack/Public/Builders/OperatorStackEditorHeaderBuilder.h
new file mode 100644
--- /dev/null
+++ b/Engine/Plugins/Experimental/OperatorStack/Source/OperatorStack/Public/Builders/OperatorStackEditorHeaderBuilder.h
@@ -0,0 +1,231 @@
+// Copyright Epic Games, Inc. All Rights Reserved.
+
+#pragma once
+
+#include "Internationalization/Text.h"
+#include "Misc/Attribute.h"
+#include "Templates/SharedPointer.h"
+#include "Styling/SlateBrush.h"
+#include "UObject/Object.h"
+
+class FCustomDetailsViewItemId;
+class FUICommandList;
+class SWidget;
+
+enum class EOperatorStackEditorMessageType
+{
+	None,
+	/** Blue message for information message */
+	Info,
+	/** Green message for success message */
+	Success,
+	/** Yellow message for warning message */
+	Warning,
+	/** Red message for error message */
+	Error
+};
+
+/** Header builder for item */
+struct OPERATORSTACKEDITOR_API FOperatorStackEditorHeaderBuilder
+{
+	/** Set the icon to display for this item header */
+	FOperatorStackEditorHeaderBuilder& SetIcon(const FSlateBrush* InIcon);
+
+	/** Set the label to display for this item header */
+	FOperatorStackEditorHeaderBuilder& SetLabel(const FText& InLabel);
+
+	/** Sets the border color surrounding the item */
+	FOperatorStackEditorHeaderBuilder& SetBorderColor(const FLinearColor& InColor);
+
+	/** Set a single property to expose in the header, like a boolean */
+	FOperatorStackEditorHeaderBuilder& SetProperty(FProperty* InProperty);
+
+	/** Set a small menu that contains icon entries on the right */
+	FOperatorStackEditorHeaderBuilder& SetToolbarMenu(const FName& InActionMenuName);
+
+	/** Set a menu that appears when we press a button with a label and icon in the header */
+	FOperatorStackEditorHeaderBuilder& SetToolMenu(const FName& InToolMenuName, const FText& InToolMenuLabel, const FSlateBrush* InToolMenuIcon);
+
+	/** Allow search on this item, if the stack header is searchable a search box will appear */
+	FOperatorStackEditorHeaderBuilder& SetSearchAllowed(bool bInSearchable);
+
+	/** Sets keywords for this item, if the keyword is contained inside this item will show */
+	FOperatorStackEditorHeaderBuilder& SetSearchKeywords(const TSet<FString>& InKeywords);
+
+	/** Sets pinned keywords on root item, they will appear on top of the search bar */
+	FOperatorStackEditorHeaderBuilder& SetSearchPinnedKeywords(const TSet<FString>& InKeywords);
+
+	/** Whether the item is expandable */
+	FOperatorStackEditorHeaderBuilder& SetExpandable(bool bInExpandable);
+
+	/** Whether the item starts expanded or collapsed */
+	FOperatorStackEditorHeaderBuilder& SetStartsExpanded(bool bInExpanded);
+
+	/** Whether the item is draggable */
+	FOperatorStackEditorHeaderBuilder& SetDraggable(bool bInDraggable);
+
+	/** Sets available commands on this item */
+	FOperatorStackEditorHeaderBuilder& SetCommandList(TSharedPtr<FUICommandList> InCommands);
+
+	/** Sets context menu name on this item */
+	FOperatorStackEditorHeaderBuilder& SetContextMenu(const FName& InToolMenuName);
+
+	/** Sets the custom widget replacing the stack generated widget */
+	FOperatorStackEditorHeaderBuilder& SetCustomWidget(TSharedPtr<SWidget> InWidget);
+
+	/** Sets whether to show a message box in the header for warning, error, success operations */
+	FOperatorStackEditorHeaderBuilder& SetMessageBox(const TAttribute<EOperatorStackEditorMessageType>& InType, const TAttribute<FText>& InText);
+
+	const FSlateBrush* GetIcon() const
+	{
+		return Icon;
+	}
+
+	const FText& GetLabel() const
+	{
+		return Label;
+	}
+
+	const FLinearColor& GetBorderColor() const
+	{
+		return BorderColor;
+	}
+
+	const FName& GetToolbarMenuName() const
+	{
+		return ToolbarMenuName;
+	}
+
+	const TSharedPtr<FCustomDetailsViewItemId>& GetProperty() const
+	{
+		return Property;
+	}
+
+	const FName& GetToolMenuName() const
+	{
+		return ToolMenuName;
+	}
+
+	const FSlateBrush* GetToolMenuIcon() const
+	{
+		return ToolMenuIcon;
+	}
+
+	const FText& GetToolMenuLabel() const
+	{
+		return ToolMenuLabel;
+	}
+
+	bool GetSearchAllowed() const
+	{
+		return bSearchAllowed;
+	}
+
+	const TSet<FString>& GetSearchKeywords() const
+	{
+		return SearchKeywords;
+	}
+
+	const TSet<FString>& GetSearchPinnedKeywords() const
+	{
+		return SearchPinnedKeywords;
+	}
+
+	bool GetExpandable() const
+	{
+		return bExpandable;
+	}
+
+	bool GetStartsExpanded() const
+	{
+		return bStartsExpanded;
+	}
+
+	bool GetDraggable() const
+	{
+		return bDraggable;
+	}
+
+	const TSharedPtr<FUICommandList>& GetCommandList() const
+	{
+		return CommandList;
+	}
+
+	const FName& GetContextMenuName() const
+	{
+		return ContextMenuName;
+	}
+
+	TSharedPtr<SWidget> GetCustomWidget() const
+	{
+		return CustomWidget;
+	}
+
+	const TAttribute<EOperatorStackEditorMessageType>& GetMessageBoxType() const
+	{
+		return MessageBoxType;
+	}
+
+	const TAttribute<FText>& GetMessageBoxText() const
+	{
+		return MessageBoxText;
+	}
+
+protected:
+	/** On the left */
+	const FSlateBrush* Icon = nullptr;
+
+	/** On the left, after to icon */
+	FText Label;
+
+	/** The border color surrounding the item */
+	FLinearColor BorderColor = FLinearColor::Transparent;
+
+	/** Menu on the right before property, should not contain a lot of entries and only icons */
+	FName ToolbarMenuName;
+
+	/** On the right, property id to expose in the header */
+	TSharedPtr<FCustomDetailsViewItemId> Property = nullptr;
+
+	/** Menu present in the header for additional actions */
+	FName ToolMenuName;
+
+	/** Menu icon in the header */
+	const FSlateBrush* ToolMenuIcon = nullptr;
+
+	/** Menu label in the header as a tooltip and label */
+	FText ToolMenuLabel;
+
+	/** Allow use of search bar to filter content */
+	bool bSearchAllowed = false;
+
+	/** Keywords to test against search input */
+	TSet<FString> SearchKeywords;
+
+	/** Pinned keywords for search, will appear below search box, only for root stack */
+	TSet<FString> SearchPinnedKeywords;
+
+	/** Allow expandable header to show the body and footer */
+	bool bExpandable = false;
+
+	/** Starts expanded or not */
+	bool bStartsExpanded = true;
+
+	/** Allow drag operation on this item */
+	bool bDraggable = false;
+
+	/** Commands executable for this item */
+	TSharedPtr<FUICommandList> CommandList;
+
+	/** Context menu available on this item */
+	FName ContextMenuName;
+
+	/** Widget that overrides the whole stack generated widget */
+	TSharedPtr<SWidget> CustomWidget;
+
+	/** The type of message box we want to show */
+	TAttribute<EOperatorStackEditorMessageType> MessageBoxType;
+
+	/** The text content of the message box */
+	TAttribute<FText> MessageBoxText;
+};
\ No newline at end of file
diff --git a/Engine/Plugins/Experimental/OperatorStack/Source/OperatorStack/Public/Contexts/OperatorStackEditorContext.h b/Engine/Plugins/Experimental/OperatorStack/Source/OperatorStack/Public/Contexts/OperatorStackEditorContext.h
new file mode 100644
--- /dev/null
+++ b/Engine/Plugins/Experimental/OperatorStack/Source/OperatorStack/Public/Contexts/OperatorStackEditorContext.h
@@ -0,0 +1,29 @@
+// Copyright Epic Games, Inc. All Rights Reserved.
+
+#pragma once
+
+#include "Items/OperatorStackEditorItem.h"
+#include "Templates/SharedPointer.h"
+
+/** Represent the current context we customize for the whole operator stack */
+struct FOperatorStackEditorContext final
+{
+	FOperatorStackEditorContext() = default;
+	explicit FOperatorStackEditorContext(const TArray<TSharedPtr<FOperatorStackEditorItem>>& InItems)
+		: Items(InItems)
+	{
+	}
+
+	/** Items we want to customize */
+	TConstArrayView<TSharedPtr<FOperatorStackEditorItem>> GetItems() const
+	{
+		return Items;
+	}
+
+protected:
+	/** Context items being customized */
+	TArray<TSharedPtr<FOperatorStackEditorItem>> Items;
+};
+
+typedef TSharedPtr<FOperatorStackEditorContext> FOperatorStackEditorContextPtr;
+typedef TWeakPtr<FOperatorStackEditorContext> FOperatorStackEditorContextPtrWeak;
\ No newline at end of file
diff --git a/Engine/Plugins/Experimental/OperatorStack/Source/OperatorStack/Public/Contexts/OperatorStackEditorMenuContext.h b/Engine/Plugins/Experimental/OperatorStack/Source/OperatorStack/Public/Contexts/OperatorStackEditorMenuContext.h
new file mode 100644
--- /dev/null
+++ b/Engine/Plugins/Experimental/OperatorStack/Source/OperatorStack/Public/Contexts/OperatorStackEditorMenuContext.h
@@ -0,0 +1,47 @@
+// Copyright Epic Games, Inc. All Rights Reserved.
+
+#pragma once
+
+#include "Templates/SharedPointer.h"
+#include "UObject/Object.h"
+#include "OperatorStackEditorMenuContext.generated.h"
+
+struct FOperatorStackEditorContext;
+struct FOperatorStackEditorItem;
+
+/** Context passed in UToolMenu when generating entries with selected items */
+UCLASS()
+class OPERATORSTACKEDITOR_API UOperatorStackEditorMenuContext : public UObject
+{
+	GENERATED_BODY()
+
+public:
+	UOperatorStackEditorMenuContext() = default;
+	
+	TSharedPtr<FOperatorStackEditorContext> GetContext() const
+	{
+		return ContextWeak.Pin();
+	}
+	
+	void SetContext(TSharedPtr<FOperatorStackEditorContext> InContext)
+	{
+		ContextWeak = InContext;
+	}
+
+	TSharedPtr<FOperatorStackEditorItem> GetItem() const
+	{
+		return ItemWeak.Pin();
+	}
+
+	void SetItem(TSharedPtr<FOperatorStackEditorItem> InItem)
+	{
+		ItemWeak = InItem;
+	}
+	
+protected:
+	/** The current context this menu is extending */
+	TWeakPtr<FOperatorStackEditorContext> ContextWeak;
+
+	/** The item this menu should apply to */
+	TWeakPtr<FOperatorStackEditorItem> ItemWeak;
+};
\ No newline at end of file
diff --git a/Engine/Plugins/Experimental/OperatorStack/Source/OperatorStack/Public/Customizations/OperatorStackEditorStackCustomization.h b/Engine/Plugins/Experimental/OperatorStack/Source/OperatorStack/Public/Customizations/OperatorStackEditorStackCustomization.h
new file mode 100644
--- /dev/null
+++ b/Engine/Plugins/Experimental/OperatorStack/Source/OperatorStack/Public/Customizations/OperatorStackEditorStackCustomization.h
@@ -0,0 +1,131 @@
+// Copyright Epic Games, Inc. All Rights Reserved.
+
+#pragma once
+
+#include "Builders/OperatorStackEditorBodyBuilder.h"
+#include "Builders/OperatorStackEditorFooterBuilder.h"
+#include "Builders/OperatorStackEditorHeaderBuilder.h"
+#include "Items/OperatorStackEditorItem.h"
+#include "Items/OperatorStackEditorTree.h"
+#include "Misc/Optional.h"
+#include "UObject/Class.h"
+#include "UObject/Field.h"
+#include "UObject/Object.h"
+#include "OperatorStackEditorStackCustomization.generated.h"
+
+enum class EItemDropZone;
+class UOperatorStackEditorStackRowCustomization;
+
+/**
+	Abstract class to represent an operator stack containing items,
+	an item is represented by a header-body-footer,
+	An item can contain multiple items (recursive),
+	Children class extending this class are automatically registered
+*/
+UCLASS(Abstract, Transient)
+class OPERATORSTACKEDITOR_API UOperatorStackEditorStackCustomization : public UObject
+{
+	GENERATED_BODY()
+
+	friend class UOperatorStackEditorSubsystem;
+	friend class SOperatorStackEditorStack;
+
+public:
+	UOperatorStackEditorStackCustomization()
+		: UOperatorStackEditorStackCustomization(NAME_None, FText::GetEmpty())
+	{}
+
+	UOperatorStackEditorStackCustomization(const FName& InIdentifier, const FText& InLabel, int32 InPriority = INDEX_NONE)
+		: Identifier(InIdentifier)
+		, Label(InLabel)
+		, Priority(InPriority)
+	{}
+
+	/** Registers items with this definition that can use this customization */
+	bool RegisterCustomizationFor(const UStruct* InItemDefinition);
+	bool RegisterCustomizationFor(const FFieldClass* InItemDefinition);
+
+	/** Unregisters item definition for this customization */
+	bool UnregisterCustomizationFor(const UStruct* InItemDefinition);
+	bool UnregisterCustomizationFor(const FFieldClass* InItemDefinition);
+
+	/** Checks if this customization is supported for this item */
+	bool IsCustomizationSupportedFor(const FOperatorStackEditorItemPtr& InItem) const;
+
+	/**
+	 * Transform a context item in items that this customization supports or returns false when not supported
+	 * For example, context item is an actor but this customization supports a specific component then look into actor for this component here
+	 */
+	virtual bool TransformContextItem(const FOperatorStackEditorItemPtr& InItem, TArray<FOperatorStackEditorItemPtr>& OutTransformedItems) const { return false; }
+
+	/** Customize the header for a context */
+	virtual void CustomizeStackHeader(const FOperatorStackEditorTree& InItemTree, FOperatorStackEditorHeaderBuilder& InStackHeaderBuilder) {}
+
+	/** Customize the full stack body */
+	virtual void CustomizeStackBody(const FOperatorStackEditorTree& InItemTree, FOperatorStackEditorBodyBuilder& InStackBodyBuilder) {}
+
+	/** Customize the header for a supported item */
+	virtual void CustomizeItemHeader(const FOperatorStackEditorItemPtr& InItem, const FOperatorStackEditorTree& InItemTree, FOperatorStackEditorHeaderBuilder& InHeaderBuilder) {}
+
+	/** Customize the body for a supported item */
+	virtual void CustomizeItemBody(const FOperatorStackEditorItemPtr& InItem, const FOperatorStackEditorTree& InItemTree, FOperatorStackEditorBodyBuilder& InBodyBuilder) {}
+
+	/** Customize the body for a supported item */
+	virtual void CustomizeItemFooter(const FOperatorStackEditorItemPtr& InItem, const FOperatorStackEditorTree& InItemTree, FOperatorStackEditorFooterBuilder& InFooterBuilder) {}
+
+	/** Can an item be dragged */
+	virtual bool OnIsItemDraggable(const FOperatorStackEditorItemPtr& InItem)
+	{
+		return false;
+	}
+
+	/** Get the valid drop zone of draggable items onto a zone item */
+	virtual TOptional<EItemDropZone> OnItemCanAcceptDrop(const TArray<FOperatorStackEditorItemPtr>& InDraggedItems, const FOperatorStackEditorItemPtr& InTargetItem, EItemDropZone InTargetZone)
+	{
+		return TOptional<EItemDropZone>();
+	}
+
+	/** Handle dropped items onto target zone */
+	virtual void OnDropItem(const TArray<FOperatorStackEditorItemPtr>& InDraggedItems, const FOperatorStackEditorItemPtr& InTargetItem, EItemDropZone InTargetZone) {}
+
+	/** Get identifier of this customization */
+	const FName& GetIdentifier() const
+	{
+		return Identifier;
+	}
+
+	/** Get display label of this customization */
+	const FText& GetLabel() const
+	{
+		return Label;
+	}
+
+	/** Get stack priority in the toolbar */
+	int32 GetPriority() const
+	{
+		return Priority;
+	}
+
+	/** Get the displayed icon of this customization */
+	virtual const FSlateBrush* GetIcon() const;
+
+private:
+	/** Unique identifier of this customization */
+	UPROPERTY()
+	FName Identifier;
+
+	/** Label displayed at the top to switch between customizations */
+	UPROPERTY()
+	FText Label;
+
+	/** Priority for this stack in toolbar, highest numbers will result in placement before lowest number */
+	UPROPERTY()
+	int32 Priority;
+
+	/** Supported definition for this customization stack */
+	UPROPERTY()
+	TSet<TObjectPtr<const UStruct>> SupportedDefinitions;
+
+	/** Supported definition for this customization stack */
+	TSet<const FFieldClass*> SupportedFieldClasses;
+};
\ No newline at end of file
diff --git a/Engine/Plugins/Experimental/OperatorStack/Source/OperatorStack/Public/Items/OperatorStackEditorItem.h b/Engine/Plugins/Experimental/OperatorStack/Source/OperatorStack/Public/Items/OperatorStackEditorItem.h
new file mode 100644
--- /dev/null
+++ b/Engine/Plugins/Experimental/OperatorStack/Source/OperatorStack/Public/Items/OperatorStackEditorItem.h
@@ -0,0 +1,152 @@
+// Copyright Epic Games, Inc. All Rights Reserved.
+
+#pragma once
+
+#include "Concepts/BaseStructureProvider.h"
+#include "Items/OperatorStackEditorItemType.h"
+#include "Templates/IsClass.h"
+#include "Templates/SharedPointer.h"
+#include "Templates/UnrealTypeTraits.h"
+#include "UObject/Class.h"
+#include "UObject/Field.h"
+#include "UObject/Object.h"
+#include "UObject/UnrealType.h"
+
+/** Abstract base parent item */
+struct FOperatorStackEditorItem
+{
+	friend class SOperatorStackEditorStack;
+
+	explicit FOperatorStackEditorItem(const FOperatorStackEditorItemType& InType)
+		: ItemType(InType)
+	{}
+
+	virtual ~FOperatorStackEditorItem() = default;
+
+	/**
+	 * You can use this to check if Item.IsA<UObject>() or Item.IsA<FMyStruct>() or Item.IsA<FIntProperty>()
+	 */
+	template<typename InValueType>
+	bool IsA() const
+	{
+		// For cases like IsA<UScriptStruct>()
+		if constexpr (TIsDerivedFrom<InValueType, UStruct>::Value)
+		{
+			if (const UStruct* ValueStruct = ItemType.Get<UStruct>())
+			{
+				return ValueStruct->IsA<InValueType>();
+			}
+		}
+
+		// Check if its a UObject derived class
+		if constexpr (TIsDerivedFrom<InValueType, UObject>::Value)
+		{
+			const UStruct* ValueStruct = InValueType::StaticClass();
+			return ItemType.IsChildOf<UStruct>(ValueStruct);
+		}
+		// Check if its a struct with StaticStruct()
+		else if constexpr (TModels<CStaticStructProvider, InValueType>::Value)
+		{
+			const UStruct* ValueStruct = InValueType::StaticStruct();
+			return ItemType.IsChildOf<UStruct>(ValueStruct);
+		}
+		// Check if base struct without StaticStruct()
+		else if constexpr (TModels<CBaseStructureProvider, InValueType>::Value)
+		{
+			const UStruct* ValueStruct = TBaseStructure<InValueType>::Get();
+			return ItemType.IsChildOf<UStruct>(ValueStruct);
+		}
+		// Check if its a FProperty derived class
+		else if constexpr (TIsDerivedFrom<InValueType, FProperty>::Value)
+		{
+			const FFieldClass* ValueClass = InValueType::StaticClass();
+			return ItemType.IsChildOf<FFieldClass>(ValueClass);
+		}
+
+		else
+		{
+			return false;
+		}
+	}
+
+	/**
+	 * You can use this to get the item Item.Get<UObject>() or Item.Get<FMyStruct>() or Item.Get<int32>()
+	 */
+	template <typename InValueType>
+	InValueType* Get() const
+	{
+		if (!GetValuePtr())
+		{
+			return nullptr;
+		}
+
+		if constexpr (TIsDerivedFrom<InValueType, UObject>::Value)
+		{
+			if (ItemType.GetTypeEnum() == EOperatorStackEditorItemType::Object)
+			{
+				return Cast<InValueType>(static_cast<UObject*>(GetValuePtr()));
+			}
+		}
+		else if constexpr (TModels_V<CStaticStructProvider, InValueType> || TModels_V<CBaseStructureProvider, InValueType>)
+		{
+			if (ItemType.GetTypeEnum() == EOperatorStackEditorItemType::Struct)
+			{
+				return static_cast<InValueType*>(GetValuePtr());
+			}
+		}
+		else if constexpr (TIsPODType<InValueType>::Value)
+		{
+			if (ItemType.GetTypeEnum() == EOperatorStackEditorItemType::Primitive)
+			{
+				return static_cast<InValueType*>(GetValuePtr());
+			}
+		}
+
+		return nullptr;
+	}
+
+	/** Get the value type of this item */
+	const FOperatorStackEditorItemType& GetValueType() const
+	{
+		return ItemType;
+	}
+
+	/** Checks if this item has a value and that it is usable */
+	virtual bool HasValue() const
+	{
+		return false;
+	}
+
+	friend uint32 GetTypeHash(const FOperatorStackEditorItem& InItem)
+	{
+		return InItem.GetHash();
+	}
+
+	bool operator==(const FOperatorStackEditorItem& InOtherItem) const
+	{
+		const uint32 ThisHash = GetHash();
+		const uint32 OtherHash = InOtherItem.GetHash();
+
+		return ThisHash != 0
+			&& OtherHash != 0
+			&& ThisHash == OtherHash;
+	}
+
+protected:
+	/** Override in child to be able to compare item */
+	virtual uint32 GetHash() const
+	{
+		return 0;
+	}
+
+	/** Get raw ptr to value, prefer using Get<>() instead */
+	virtual void* GetValuePtr() const
+	{
+		return nullptr;
+	}
+
+	FOperatorStackEditorItemType ItemType;
+};
+
+typedef TSharedPtr<FOperatorStackEditorItem> FOperatorStackEditorItemPtr;
+typedef TWeakPtr<FOperatorStackEditorItem> FOperatorStackEditorItemPtrWeak;
diff --git a/Engine/Plugins/Experimental/OperatorStack/Source/OperatorStack/Public/Items/OperatorStackEditorItemType.h b/Engine/Plugins/Experimental/OperatorStack/Source/OperatorStack/Public/Items/OperatorStackEditorItemType.h
new file mode 100644
--- /dev/null
+++ b/Engine/Plugins/Experimental/OperatorStack/Source/OperatorStack/Public/Items/OperatorStackEditorItemType.h
@@ -0,0 +1,111 @@
+// Copyright Epic Games, Inc. All Rights Reserved.
+
+#pragma once
+
+#include <type_traits>
+
+#include "Containers/Union.h"
+#include "Templates/UnrealTypeTraits.h"
+#include "UObject/Class.h"
+#include "UObject/Field.h"
+#include "UObject/Object.h"
+#include "UObject/UnrealType.h"
+
+/** Used to quickly get the general type of an item without casting */
+enum class EOperatorStackEditorItemType : uint8
+{
+	None,
+	Object,
+	Struct,
+	Primitive
+};
+
+/** Item type */
+struct FOperatorStackEditorItemType
+{
+	template<typename InTypeClass
+		UE_REQUIRES(std::is_base_of_v<UStruct, InTypeClass> || std::is_base_of_v<FFieldClass, InTypeClass>)>
+	explicit FOperatorStackEditorItemType(const InTypeClass* InType, EOperatorStackEditorItemType InTypeEnum)
+		: TypeClass(InType)
+		, TypeEnum(InTypeEnum)
+	{
+	}
+
+	template<typename InTypeClass
+		UE_REQUIRES(std::is_base_of_v<UStruct, InTypeClass> || std::is_base_of_v<FFieldClass, InTypeClass>)>
+	bool IsChildOf(const InTypeClass* InOtherType) const
+	{
+		if (TypeClass.HasSubtype<const InTypeClass*>())
+		{
+			return TypeClass.GetSubtype<const InTypeClass*>()->IsChildOf(InOtherType);
+		}
+
+		return false;
+	}
+
+	bool IsChildOf(const FOperatorStackEditorItemType* InOtherType) const
+	{
+		if (!InOtherType)
+		{
+			return false;
+		}
+
+		if (TypeClass.GetCurrentSubtypeIndex() != InOtherType->TypeClass.GetCurrentSubtypeIndex())
+		{
+			return false;
+		}
+
+		if (TypeClass.GetCurrentSubtypeIndex() == 0)
+		{
+			return TypeClass.GetSubtype<const UStruct*>()->IsChildOf(InOtherType->TypeClass.GetSubtype<const UStruct*>());
+		}
+
+		if (TypeClass.GetCurrentSubtypeIndex() == 1)
+		{
+			return TypeClass.GetSubtype<const FFieldClass*>()->IsChildOf(InOtherType->TypeClass.GetSubtype<const FFieldClass*>());
+		}
+
+		return false;
+	}
+
+	template<typename InTypeClass
+		UE_REQUIRES(std::is_base_of_v<UStruct, InTypeClass> || std::is_base_of_v<FFieldClass, InTypeClass>)>
+	const InTypeClass* Get() const
+	{
+		return TypeClass.HasSubtype<const InTypeClass*>() ? TypeClass.GetSubtype<const InTypeClass*>() : nullptr;
+	}
+
+	EOperatorStackEditorItemType GetTypeEnum() const
+	{
+		return TypeEnum;
+	}
+
+	friend uint32 GetTypeHash(const FOperatorStackEditorItemType& InItem)
+	{
+		uint32 TypeClassHash = 0;
+
+		if (InItem.TypeClass.GetCurrentSubtypeIndex() == 0)
+		{
+			TypeClassHash = GetTypeHash(InItem.TypeClass.GetSubtype<const UStruct*>());
+		}
+
+		if (InItem.TypeClass.GetCurrentSubtypeIndex() == 1)
+		{
+			TypeClassHash = GetTypeHash(InItem.TypeClass.GetSubtype<const FFieldClass*>());
+		}
+
+		return HashCombine(TypeClassHash, static_cast<uint32>(InItem.TypeEnum));
+	}
+
+	bool operator==(const FOperatorStackEditorItemType& InOtherItem) const
+	{
+		return TypeEnum == InOtherItem.TypeEnum 
+			&& TypeClass.GetCurrentSubtypeIndex() == InOtherItem.TypeClass.GetCurrentSubtypeIndex()
+			&& ((TypeClass.GetCurrentSubtypeIndex() == 0 && TypeClass.GetSubtype<const UStruct*>() == InOtherItem.TypeClass.GetSubtype<const UStruct*>())
+			|| (TypeClass.GetCurrentSubtypeIndex() == 1 && TypeClass.GetSubtype<const FFieldClass*>() == InOtherItem.TypeClass.GetSubtype<const FFieldClass*>()));
+	}
+
+private:
+	TUnion<const UStruct*, const FFieldClass*> TypeClass;
+	EOperatorStackEditorItemType TypeEnum = EOperatorStackEditorItemType::None;
+};
\ No newline at end of file
diff --git a/Engine/Plugins/Experimental/OperatorStack/Source/OperatorStack/Public/Items/OperatorStackEditorObjectItem.h b/Engine/Plugins/Experimental/OperatorStack/Source/OperatorStack/Public/Items/OperatorStackEditorObjectItem.h
new file mode 100644
--- /dev/null
+++ b/Engine/Plugins/Experimental/OperatorStack/Source/OperatorStack/Public/Items/OperatorStackEditorObjectItem.h
@@ -0,0 +1,46 @@
+// Copyright Epic Games, Inc. All Rights Reserved.
+
+#pragma once
+
+#include "OperatorStackEditorItem.h"
+#include "Templates/SharedPointer.h"
+#include "Templates/UnrealTypeTraits.h"
+#include "UObject/Class.h"
+#include "UObject/Field.h"
+#include "UObject/Object.h"
+#include "UObject/UnrealType.h"
+
+/** Object version of item */
+struct FOperatorStackEditorObjectItem : FOperatorStackEditorItem
+{
+	template<
+		typename InValueType
+		UE_REQUIRES(TIsDerivedFrom<InValueType, UObject>::Value)>
+	explicit FOperatorStackEditorObjectItem(InValueType* InItem)
+		: FOperatorStackEditorItem(
+			FOperatorStackEditorItemType(InItem ? InItem->GetClass() : nullptr, EOperatorStackEditorItemType::Object)
+		)
+	{
+		ObjectWeak = InItem;
+		CachedHash = GetTypeHash(InItem);
+	}
+
+	virtual bool HasValue() const override
+	{
+		return ObjectWeak.IsValid();
+	}
+
+protected:
+	virtual uint32 GetHash() const override
+	{
+		return CachedHash;
+	}
+
+	virtual void* GetValuePtr() const override
+	{
+		return ObjectWeak.Get();
+	}
+
+	uint32 CachedHash = 0;
+	TWeakObjectPtr<UObject> ObjectWeak;
+};
\ No newline at end of file
diff --git a/Engine/Plugins/Experimental/OperatorStack/Source/OperatorStack/Public/Items/OperatorStackEditorPrimitiveItem.h b/Engine/Plugins/Experimental/OperatorStack/Source/OperatorStack/Public/Items/OperatorStackEditorPrimitiveItem.h
new file mode 100644
--- /dev/null
+++ b/Engine/Plugins/Experimental/OperatorStack/Source/OperatorStack/Public/Items/OperatorStackEditorPrimitiveItem.h
@@ -0,0 +1,46 @@
+// Copyright Epic Games, Inc. All Rights Reserved.
+
+#pragma once
+
+#include "OperatorStackEditorItem.h"
+#include "Templates/SharedPointer.h"
+#include "Templates/UnrealTypeTraits.h"
+#include "UObject/Class.h"
+#include "UObject/Field.h"
+#include "UObject/Object.h"
+#include "UObject/UnrealType.h"
+
+/** Primitive version of item */
+struct FOperatorStackEditorPrimitiveItem : FOperatorStackEditorItem
+{
+	template<
+		typename InItemDefinition
+		UE_REQUIRES(TIsPODType<InItemDefinition>::Value)>
+	explicit FOperatorStackEditorPrimitiveItem(FFieldClass* InFieldClass, InItemDefinition& InValue)
+		: FOperatorStackEditorItem(
+			FOperatorStackEditorItemType(InFieldClass, EOperatorStackEditorItemType::Primitive)
+		)
+	{
+		MemoryPtr = reinterpret_cast<uint8*>(&InValue);
+		CachedHash = GetTypeHash(InValue);
+	}
+
+	virtual bool HasValue() const override
+	{
+		return MemoryPtr != nullptr;
+	}
+
+protected:
+	virtual uint32 GetHash() const override
+	{
+		return CachedHash;
+	}
+
+	virtual void* GetValuePtr() const override
+	{
+		return MemoryPtr;
+	}
+
+	uint32 CachedHash = 0;
+	uint8* MemoryPtr = nullptr;
+};
\ No newline at end of file
diff --git a/Engine/Plugins/Experimental/OperatorStack/Source/OperatorStack/Public/Items/OperatorStackEditorStructItem.h b/Engine/Plugins/Experimental/OperatorStack/Source/OperatorStack/Public/Items/OperatorStackEditorStructItem.h
new file mode 100644
--- /dev/null
+++ b/Engine/Plugins/Experimental/OperatorStack/Source/OperatorStack/Public/Items/OperatorStackEditorStructItem.h
@@ -0,0 +1,96 @@
+// Copyright Epic Games, Inc. All Rights Reserved.
+
+#pragma once
+
+#include "OperatorStackEditorItem.h"
+#include "Templates/SharedPointer.h"
+#include "Templates/UnrealTypeTraits.h"
+#include "UObject/Class.h"
+#include "UObject/Field.h"
+#include "UObject/Object.h"
+#include "UObject/StructOnScope.h"
+#include "UObject/UnrealType.h"
+
+/** Struct version of item */
+struct FOperatorStackEditorStructItem : FOperatorStackEditorItem
+{
+	// Version for UStruct
+	template <typename InValueType
+		UE_REQUIRES(TModels<CStaticStructProvider, InValueType>::Value)>
+	explicit FOperatorStackEditorStructItem(InValueType& InStruct)
+		: FOperatorStackEditorItem(
+			FOperatorStackEditorItemType(InValueType::StaticStruct(), EOperatorStackEditorItemType::Struct)
+		)
+	{
+		StructWeak = MakeShared<FStructOnScope>(InValueType::StaticStruct(), reinterpret_cast<uint8*>(&InStruct));
+
+		if constexpr (TModels<CGetTypeHashable, InValueType>::Value)
+		{
+			CachedHash = GetTypeHash(InStruct);
+		}
+	}
+
+	// Version for base struct
+	template <typename InValueType
+		UE_REQUIRES(TModels<CBaseStructureProvider, InValueType>::Value)>
+	explicit FOperatorStackEditorStructItem(InValueType& InStruct)
+		: FOperatorStackEditorItem(
+			FOperatorStackEditorItemType(TBaseStructure<InValueType>::Get(), EOperatorStackEditorItemType::Struct)
+		)
+	{
+		StructWeak = MakeShared<FStructOnScope>(TBaseStructure<InValueType>::Get(), reinterpret_cast<uint8*>(&InStruct));
+
+		if constexpr (TModels<CGetTypeHashable, InValueType>::Value)
+		{
+			CachedHash = GetTypeHash(InStruct);
+		}
+	}
+
+	// Version for struct on scope
+	explicit FOperatorStackEditorStructItem(const TSharedPtr<FStructOnScope>& InStruct)
+		: FOperatorStackEditorItem(
+			FOperatorStackEditorItemType(InStruct.IsValid() ? InStruct->GetStruct() : nullptr, EOperatorStackEditorItemType::Struct)
+		)
+	{
+		StructWeak = InStruct;
+
+		if (InStruct.IsValid())
+		{
+			CachedHash = HashCombine(GetTypeHash(ItemType), GetTypeHash(*InStruct->GetStructMemory()));
+		}
+	}
+
+	TSharedPtr<FStructOnScope> GetStructOnScope() const
+	{
+		return StructWeak.Pin();
+	}
+
+	virtual bool HasValue() const override
+	{
+		if (const TSharedPtr<FStructOnScope> Struct = StructWeak.Pin())
+		{
+			return Struct->IsValid();
+		}
+
+		return false;
+	}
+
+protected:
+	virtual uint32 GetHash() const override
+	{
+		return CachedHash;
+	}
+
+	virtual void* GetValuePtr() const override
+	{
+		if (const TSharedPtr<FStructOnScope> Struct = StructWeak.Pin())
+		{
+			return Struct->GetStructMemory();
+		}
+
+		return nullptr;
+	}
+
+	uint32 CachedHash = 0;
+	TWeakPtr<FStructOnScope> StructWeak;
+};
\ No newline at end of file
diff --git a/Engine/Plugins/Experimental/OperatorStack/Source/OperatorStack/Public/Items/OperatorStackEditorTree.h b/Engine/Plugins/Experimental/OperatorStack/Source/OperatorStack/Public/Items/OperatorStackEditorTree.h
new file mode 100644
--- /dev/null
+++ b/Engine/Plugins/Experimental/OperatorStack/Source/OperatorStack/Public/Items/OperatorStackEditorTree.h
@@ -0,0 +1,72 @@
+// Copyright Epic Games, Inc. All Rights Reserved.
+
+#pragma once
+
+#include "Contexts/OperatorStackEditorContext.h"
+#include "OperatorStackEditorItem.h"
+
+class UOperatorStackEditorStackCustomization;
+
+/** Customization tree used to find all supported items by a customization object */
+struct OPERATORSTACKEDITOR_API FOperatorStackEditorTree
+{
+	explicit FOperatorStackEditorTree(UOperatorStackEditorStackCustomization* InCustomization, FOperatorStackEditorContextPtr InContext);
+
+	/** Get all top items supported by this customization */
+	TArray<FOperatorStackEditorItemPtr> GetRootItems() const;
+
+	/** Get supported children items from a supported item */
+	TArray<FOperatorStackEditorItemPtr> GetChildrenItems(FOperatorStackEditorItemPtr InItem) const;
+
+	/** Get supported parent item from a supported child item */
+	FOperatorStackEditorItemPtr GetParentItem(FOperatorStackEditorItemPtr InItem) const;
+
+	/** Get all leaf supported item */
+	TArray<FOperatorStackEditorItemPtr> GetLeafItems() const;
+
+	/** Get all supported items by this customization */
+	TConstArrayView<FOperatorStackEditorItemPtr> GetAllItems() const;
+
+	/** Get source context used to build this tree */
+	const FOperatorStackEditorContext& GetContext() const;
+
+	/** Get the customization used to build this tree */
+	UOperatorStackEditorStackCustomization* GetCustomization() const;
+
+	/** Checks if an item is already contained in the tree */
+	bool Contains(FOperatorStackEditorItemPtr InItem) const;
+
+private:
+	/** Customization tree node used only internally by the tree */
+	struct FOperatorStackEditorTreeNode
+	{
+		explicit FOperatorStackEditorTreeNode(int32 InItemIndex, int32 InParentIndex);
+
+		int32 ItemIndex = INDEX_NONE;
+		int32 ParentIndex = INDEX_NONE;
+		TArray<int32> ChildrenIndices;
+	};
+
+	using FOperatorStackEditorTreeNodePtr = TSharedPtr<FOperatorStackEditorTreeNode>;
+
+	void GetSupportedItems(const TArray<FOperatorStackEditorItemPtr>& InItems, TArray<FOperatorStackEditorItemPtr>& OutSupportedItems) const;
+	TArray<FOperatorStackEditorItemPtr> GetSupportedChildrenItems(const FOperatorStackEditorItemPtr& InParentItem) const;
+
+	void BuildTree(const TArray<FOperatorStackEditorItemPtr>& InSourceItems);
+	void BuildTreeInternal(const TArray<FOperatorStackEditorItemPtr>& InItems, int32 InParent);
+
+	/** All items in this tree */
+	TArray<FOperatorStackEditorItemPtr> Items;
+
+	/** All nodes in this tree */
+	TArray<FOperatorStackEditorTreeNodePtr> Nodes;
+
+	/** Root items in the tree */
+	TArray<FOperatorStackEditorTreeNodePtr> RootNodes;
+
+	/** Customization used to build this tree */
+	TWeakObjectPtr<UOperatorStackEditorStackCustomization> CustomizationWeak;
+
+	/** Source context used to build this tree */
+	TWeakPtr<FOperatorStackEditorContext> ContextWeak;
+};
diff --git a/Engine/Plugins/Experimental/OperatorStack/Source/OperatorStack/Public/Subsystems/OperatorStackEditorSubsystem.h b/Engine/Plugins/Experimental/OperatorStack/Source/OperatorStack/Public/Subsystems/OperatorStackEditorSubsystem.h
new file mode 100644
--- /dev/null
+++ b/Engine/Plugins/Experimental/OperatorStack/Source/OperatorStack/Public/Subsystems/OperatorStackEditorSubsystem.h
@@ -0,0 +1,58 @@
+// Copyright Epic Games, Inc. All Rights Reserved.
+
+#pragma once
+
+#include "EditorSubsystem.h"
+#include "Templates/SubclassOf.h"
+#include "UObject/WeakObjectPtrTemplates.h"
+#include "OperatorStackEditorSubsystem.generated.h"
+
+class SOperatorStackEditorWidget;
+class UOperatorStackEditorStackCustomization;
+
+/** Subsystem that handles operator stack customization */
+UCLASS()
+class OPERATORSTACKEDITOR_API UOperatorStackEditorSubsystem : public UEditorSubsystem
+{
+	GENERATED_BODY()
+
+	friend class SOperatorStackEditorPanel;
+	
+public:
+	static UOperatorStackEditorSubsystem* Get();
+	
+	//~ Begin USubsystem
+	virtual void Initialize(FSubsystemCollectionBase& Collection) override;
+	virtual void Deinitialize() override;
+	//~ End USubsystem
+
+	bool RegisterStackCustomization(TSubclassOf<UOperatorStackEditorStackCustomization> InStackCustomizationClass);
+	bool UnregisterStackCustomization(TSubclassOf<UOperatorStackEditorStackCustomization> InStackCustomizationClass);
+
+	/** Loops through each customization stack in priority order, stops if false returned within function */
+	bool ForEachCustomization(TFunctionRef<bool(UOperatorStackEditorStackCustomization*)> InFunction) const;
+
+	/** Get customization with identifier */
+	UOperatorStackEditorStackCustomization* GetCustomization(const FName& InName) const;
+	
+	/** Generate an operator stack widget */
+	TSharedRef<SOperatorStackEditorWidget> GenerateWidget();
+
+	/** Finds an existing operator stack widget */
+	TSharedPtr<SOperatorStackEditorWidget> FindWidget(int32 InId);
+	
+	/** Loops through each operator stack widgets available */
+	bool ForEachCustomizationWidget(TFunctionRef<bool(TSharedRef<SOperatorStackEditorWidget>)> InFunction) const;
+
+protected:
+    void ScanForStackCustomizations();
+
+	void OnWidgetDestroyed(int32 InPanelId);
+
+	/** Map of identifier and stack customization models */
+	UPROPERTY()
+	TMap<FName, TObjectPtr<UOperatorStackEditorStackCustomization>> CustomizationStacks;
+
+	/** Current widget created associated to an identifier */
+	TMap<int32, TWeakPtr<SOperatorStackEditorWidget>> CustomizationWidgets;
+};
diff --git a/Engine/Plugins/Experimental/OperatorStack/Source/OperatorStack/Public/Widgets/SOperatorStackEditorWidget.h b/Engine/Plugins/Experimental/OperatorStack/Source/OperatorStack/Public/Widgets/SOperatorStackEditorWidget.h
new file mode 100644
--- /dev/null
+++ b/Engine/Plugins/Experimental/OperatorStack/Source/OperatorStack/Public/Widgets/SOperatorStackEditorWidget.h
@@ -0,0 +1,63 @@
+// Copyright Epic Games, Inc. All Rights Reserved.
+
+#pragma once
+
+#include "Templates/SharedPointer.h"
+#include "Widgets/SCompoundWidget.h"
+
+class FDetailColumnSizeData;
+class IDetailKeyframeHandler;
+struct FOperatorStackEditorContext;
+
+/** Main operator stack widget created, only exposed function */
+class SOperatorStackEditorWidget : public SCompoundWidget
+{
+public:
+	/** Set the active customization in the panel */
+	virtual void SetActiveCustomization(const FName& InCustomization) {}
+
+	/** Sets the customization to show in the toolbar, overrides order */
+	virtual void SetToolbarCustomization(const TArray<FName>& InCustomizations) {}
+
+	/** Show the customization in the panel */
+	virtual void ShowToolbarCustomization(FName InCustomization) {}
+
+	/** Hide the customization from the panel */
+	virtual void HideToolbarCustomization(FName InCustomization) {}
+
+	/** Sets the visibility of the toolbar to switch between customizations */
+	virtual void SetToolbarVisibility(bool bInVisible) {}
+
+	/** Update the context of the widget panel */
+	virtual void SetContext(const FOperatorStackEditorContext& InNewContext) {}
+
+	/** Sets the keyframe handler to keyframe items properties */
+	virtual void SetKeyframeHandler(TSharedPtr<IDetailKeyframeHandler> InKeyframeHandler) {}
+
+	/** Sets the detail column sizes to have same column sizes across items */
+	virtual void SetDetailColumnSize(TSharedPtr<FDetailColumnSizeData> InDetailColumnSize) {}
+
+	/** Sets the panel tag to identify this panel */
+	virtual void SetPanelTag(FName InTag) {}
+
+	/** Forces a refresh of the current context if any */
+	virtual void RefreshContext() {}
+
+	/** Get the current context of this operator stack widget */
+	virtual TSharedPtr<FOperatorStackEditorContext> GetContext() const
+	{
+		return nullptr;
+	}
+
+	/** Get the panel identifier, unique per widget */
+	virtual int32 GetPanelId() const
+	{
+		return INDEX_NONE;
+	}
+
+	/** Get the panel tag */
+	virtual FName GetPanelTag() const
+	{
+		return NAME_None;
+	}
+};
\ No newline at end of file
