Motion Design:
- Operator Stack to use minimal api
- comments added to transition tree

#changelist validated
#virtualized

diff --git a/Engine/Plugins/Experimental/Avalanche/Source/AvalancheTransition/Public/AvaTransitionTree.h b/Engine/Plugins/Experimental/Avalanche/Source/AvalancheTransition/Public/AvaTransitionTree.h
--- a/Engine/Plugins/Experimental/Avalanche/Source/AvalancheTransition/Public/AvaTransitionTree.h
+++ b/Engine/Plugins/Experimental/Avalanche/Source/AvalancheTransition/Public/AvaTransitionTree.h
@@ -6,6 +6,10 @@
 #include "StateTree.h"
 #include "AvaTransitionTree.generated.h"
 
+/**
+ * Motion Design Transition Tree is a State Tree with the purpose of executing user-defined logic
+ * when there's a Transition between multiple scenes in multiple layers.
+ */
 UCLASS(DisplayName = "Motion Design Transition Tree")
 class AVALANCHETRANSITION_API UAvaTransitionTree : public UStateTree
 {
diff --git a/Engine/Plugins/Experimental/OperatorStack/Source/OperatorStack/Public/Contexts/OperatorStackEditorMenuContext.h b/Engine/Plugins/Experimental/OperatorStack/Source/OperatorStack/Public/Contexts/OperatorStackEditorMenuContext.h
--- a/Engine/Plugins/Experimental/OperatorStack/Source/OperatorStack/Public/Contexts/OperatorStackEditorMenuContext.h
+++ b/Engine/Plugins/Experimental/OperatorStack/Source/OperatorStack/Public/Contexts/OperatorStackEditorMenuContext.h
@@ -10,19 +10,17 @@ struct FOperatorStackEditorContext;
 struct FOperatorStackEditorItem;
 
 /** Context passed in UToolMenu when generating entries with selected items */
-UCLASS()
-class OPERATORSTACKEDITOR_API UOperatorStackEditorMenuContext : public UObject
+UCLASS(MinimalAPI)
+class UOperatorStackEditorMenuContext : public UObject
 {
 	GENERATED_BODY()
 
 public:
-	UOperatorStackEditorMenuContext() = default;
-	
 	TSharedPtr<FOperatorStackEditorContext> GetContext() const
 	{
 		return ContextWeak.Pin();
 	}
-	
+
 	void SetContext(TSharedPtr<FOperatorStackEditorContext> InContext)
 	{
 		ContextWeak = InContext;
@@ -37,7 +35,7 @@ public:
 	{
 		ItemWeak = InItem;
 	}
-	
+
 protected:
 	/** The current context this menu is extending */
 	TWeakPtr<FOperatorStackEditorContext> ContextWeak;
