MotionDesign : OperatorStack
- Cleaning public facing API

#changelist validated
#virtualized

diff --git a/Engine/Plugins/Experimental/OperatorStack/Source/OperatorStack/Public/Customizations/OperatorStackEditorStackCustomization.h b/Engine/Plugins/Experimental/OperatorStack/Source/OperatorStack/Public/Customizations/OperatorStackEditorStackCustomization.h
--- a/Engine/Plugins/Experimental/OperatorStack/Source/OperatorStack/Public/Customizations/OperatorStackEditorStackCustomization.h
+++ b/Engine/Plugins/Experimental/OperatorStack/Source/OperatorStack/Public/Customizations/OperatorStackEditorStackCustomization.h
@@ -22,8 +22,8 @@ class UOperatorStackEditorStackRowCustomization;
 	An item can contain multiple items (recursive),
 	Children class extending this class are automatically registered
 */
-UCLASS(Abstract, Transient)
-class OPERATORSTACKEDITOR_API UOperatorStackEditorStackCustomization : public UObject
+UCLASS(MinimalAPI, Abstract, Transient)
+class UOperatorStackEditorStackCustomization : public UObject
 {
 	GENERATED_BODY()
 
@@ -42,12 +42,12 @@ public:
 	{}
 
 	/** Registers items with this definition that can use this customization */
-	bool RegisterCustomizationFor(const UStruct* InItemDefinition);
-	bool RegisterCustomizationFor(const FFieldClass* InItemDefinition);
+	OPERATORSTACKEDITOR_API bool RegisterCustomizationFor(const UStruct* InItemDefinition);
+	OPERATORSTACKEDITOR_API bool RegisterCustomizationFor(const FFieldClass* InItemDefinition);
 
 	/** Unregisters item definition for this customization */
-	bool UnregisterCustomizationFor(const UStruct* InItemDefinition);
-	bool UnregisterCustomizationFor(const FFieldClass* InItemDefinition);
+	OPERATORSTACKEDITOR_API bool UnregisterCustomizationFor(const UStruct* InItemDefinition);
+	OPERATORSTACKEDITOR_API bool UnregisterCustomizationFor(const FFieldClass* InItemDefinition);
 
 	/** Checks if this customization is supported for this item */
 	bool IsCustomizationSupportedFor(const FOperatorStackEditorItemPtr& InItem) const;
@@ -107,7 +107,7 @@ public:
 	}
 
 	/** Get the displayed icon of this customization */
-	virtual const FSlateBrush* GetIcon() const;
+	OPERATORSTACKEDITOR_API virtual const FSlateBrush* GetIcon() const;
 
 private:
 	/** Unique identifier of this customization */
diff --git a/Engine/Plugins/Experimental/OperatorStack/Source/OperatorStack/Public/Items/OperatorStackEditorTree.h b/Engine/Plugins/Experimental/OperatorStack/Source/OperatorStack/Public/Items/OperatorStackEditorTree.h
--- a/Engine/Plugins/Experimental/OperatorStack/Source/OperatorStack/Public/Items/OperatorStackEditorTree.h
+++ b/Engine/Plugins/Experimental/OperatorStack/Source/OperatorStack/Public/Items/OperatorStackEditorTree.h
@@ -8,12 +8,12 @@
 class UOperatorStackEditorStackCustomization;
 
 /** Customization tree used to find all supported items by a customization object */
-struct OPERATORSTACKEDITOR_API FOperatorStackEditorTree
+struct FOperatorStackEditorTree
 {
 	explicit FOperatorStackEditorTree(UOperatorStackEditorStackCustomization* InCustomization, FOperatorStackEditorContextPtr InContext);
 
 	/** Get all top items supported by this customization */
-	TArray<FOperatorStackEditorItemPtr> GetRootItems() const;
+	OPERATORSTACKEDITOR_API TArray<FOperatorStackEditorItemPtr> GetRootItems() const;
 
 	/** Get supported children items from a supported item */
 	TArray<FOperatorStackEditorItemPtr> GetChildrenItems(FOperatorStackEditorItemPtr InItem) const;
@@ -25,10 +25,10 @@ struct OPERATORSTACKEDITOR_API FOperatorStackEditorTree
 	TArray<FOperatorStackEditorItemPtr> GetLeafItems() const;
 
 	/** Get all supported items by this customization */
-	TConstArrayView<FOperatorStackEditorItemPtr> GetAllItems() const;
+	OPERATORSTACKEDITOR_API TConstArrayView<FOperatorStackEditorItemPtr> GetAllItems() const;
 
 	/** Get source context used to build this tree */
-	const FOperatorStackEditorContext& GetContext() const;
+	OPERATORSTACKEDITOR_API const FOperatorStackEditorContext& GetContext() const;
 
 	/** Get the customization used to build this tree */
 	UOperatorStackEditorStackCustomization* GetCustomization() const;
diff --git a/Engine/Plugins/Experimental/OperatorStack/Source/OperatorStack/Public/Subsystems/OperatorStackEditorSubsystem.h b/Engine/Plugins/Experimental/OperatorStack/Source/OperatorStack/Public/Subsystems/OperatorStackEditorSubsystem.h
--- a/Engine/Plugins/Experimental/OperatorStack/Source/OperatorStack/Public/Subsystems/OperatorStackEditorSubsystem.h
+++ b/Engine/Plugins/Experimental/OperatorStack/Source/OperatorStack/Public/Subsystems/OperatorStackEditorSubsystem.h
@@ -12,37 +12,37 @@ class UOperatorStackEditorStackCustomization;
 
 /** Subsystem that handles operator stack customization */
 UCLASS()
-class OPERATORSTACKEDITOR_API UOperatorStackEditorSubsystem : public UEditorSubsystem
+class UOperatorStackEditorSubsystem : public UEditorSubsystem
 {
 	GENERATED_BODY()
 
 	friend class SOperatorStackEditorPanel;
-	
+
 public:
-	static UOperatorStackEditorSubsystem* Get();
-	
+	OPERATORSTACKEDITOR_API static UOperatorStackEditorSubsystem* Get();
+
 	//~ Begin USubsystem
 	virtual void Initialize(FSubsystemCollectionBase& Collection) override;
 	virtual void Deinitialize() override;
 	//~ End USubsystem
 
-	bool RegisterStackCustomization(TSubclassOf<UOperatorStackEditorStackCustomization> InStackCustomizationClass);
-	bool UnregisterStackCustomization(TSubclassOf<UOperatorStackEditorStackCustomization> InStackCustomizationClass);
+	OPERATORSTACKEDITOR_API bool RegisterStackCustomization(TSubclassOf<UOperatorStackEditorStackCustomization> InStackCustomizationClass);
+	OPERATORSTACKEDITOR_API bool UnregisterStackCustomization(TSubclassOf<UOperatorStackEditorStackCustomization> InStackCustomizationClass);
 
 	/** Loops through each customization stack in priority order, stops if false returned within function */
 	bool ForEachCustomization(TFunctionRef<bool(UOperatorStackEditorStackCustomization*)> InFunction) const;
 
 	/** Get customization with identifier */
 	UOperatorStackEditorStackCustomization* GetCustomization(const FName& InName) const;
-	
+
 	/** Generate an operator stack widget */
-	TSharedRef<SOperatorStackEditorWidget> GenerateWidget();
+	OPERATORSTACKEDITOR_API TSharedRef<SOperatorStackEditorWidget> GenerateWidget();
 
 	/** Finds an existing operator stack widget */
 	TSharedPtr<SOperatorStackEditorWidget> FindWidget(int32 InId);
-	
+
 	/** Loops through each operator stack widgets available */
-	bool ForEachCustomizationWidget(TFunctionRef<bool(TSharedRef<SOperatorStackEditorWidget>)> InFunction) const;
+	OPERATORSTACKEDITOR_API bool ForEachCustomizationWidget(TFunctionRef<bool(TSharedRef<SOperatorStackEditorWidget>)> InFunction) const;
 
 protected:
     void ScanForStackCustomizations();
