MotionDesign : OperatorStack
- Allow item body to have a widget and nested children

diff --git a/Engine/Plugins/Experimental/OperatorStack/Source/OperatorStack/Private/Widgets/SOperatorStackEditorStack.cpp b/Engine/Plugins/Experimental/OperatorStack/Source/OperatorStack/Private/Widgets/SOperatorStackEditorStack.cpp
--- a/Engine/Plugins/Experimental/OperatorStack/Source/OperatorStack/Private/Widgets/SOperatorStackEditorStack.cpp
+++ b/Engine/Plugins/Experimental/OperatorStack/Source/OperatorStack/Private/Widgets/SOperatorStackEditorStack.cpp
@@ -464,6 +464,45 @@ TSharedPtr<SWidget> SOperatorStackEditorStack::GenerateBodyWidget()
 
 	FText EmptyBodyDefaultText = LOCTEXT("EmptyBodyText", "Select a supported item to display it here");
 
+	auto GetItemWidget = [this, MainPanel](const FOperatorStackEditorBodyBuilder& InItem)->TSharedPtr<SWidget>
+	{
+		// We have set a custom widget for this body
+		if (TSharedPtr<SWidget> CustomWidget = InItem.GetCustomWidget())
+		{
+			return CustomWidget;
+		}
+
+		// Only build details view if we allow it
+		if (InItem.GetShowDetailsView())
+		{
+			FCustomDetailsViewArgs BodyCustomViewArgs;
+			BodyCustomViewArgs.bShowCategories = false;
+			BodyCustomViewArgs.KeyframeHandler = MainPanel->GetKeyframeHandler();
+			BodyCustomViewArgs.bAllowGlobalExtensions = true;
+			BodyCustomViewArgs.ColumnSizeData = MainPanel->GetDetailColumnSize();
+
+			for (const TSharedPtr<FCustomDetailsViewItemId>& DetailsViewId : InItem.GetDisallowedDetailsViewItems())
+			{
+				BodyCustomViewArgs.ItemAllowList.Disallow(*DetailsViewId.Get());
+			}
+
+			for (const TSharedPtr<FCustomDetailsViewItemId>& DetailsViewId : InItem.GetAllowedDetailsViewItems())
+			{
+				BodyCustomViewArgs.ItemAllowList.Allow(*DetailsViewId.Get());
+			}
+
+			const FOperatorStackEditorItem* DetailViewItem = InItem.GetDetailsViewItem().IsValid()
+				? InItem.GetDetailsViewItem().Get()
+				: CustomizeItem.Get();
+
+			BodyDetailsView = CreateDetailsView(BodyCustomViewArgs, *DetailViewItem); 
+
+			return BodyDetailsView;
+		}
+
+		return nullptr;
+	};
+
 	// We are the root and we have multiple supported items selected
 	if (!CustomizeItem.IsValid() && !Items.IsEmpty())
 	{
@@ -507,21 +546,41 @@ TSharedPtr<SWidget> SOperatorStackEditorStack::GenerateBodyWidget()
 	// We are not the root but we contain children then add a list view
 	if (CustomizeItem.IsValid() && !Items.IsEmpty())
 	{
-		ItemsListView = SNew(SListView<FOperatorStackEditorItemPtr>)
-			.ListViewStyle(&FOperatorStackEditorStyle::Get().GetWidgetStyle<FTableViewStyle>("ListViewStyle"))
-			.ListItemsSource(&Items)
-			.ClearSelectionOnClick(true)
-			.SelectionMode(ESelectionMode::Multi)
-			.OnKeyDownHandler(this, &SOperatorStackEditorStack::OnKeyDownHandler)
-			.OnSelectionChanged(this, &SOperatorStackEditorStack::OnSelectionChanged)
-			.OnGenerateRow(this, &SOperatorStackEditorStack::OnGenerateRow);
+		FOperatorStackEditorBodyBuilder ItemBodyBuilder;
+		StackCustomization->CustomizeItemBody(CustomizeItem, ItemTree, ItemBodyBuilder);
+
+		TSharedRef<SVerticalBox> ParentBox = SNew(SVerticalBox);
+
+		if (TSharedPtr<SWidget> BodyWidget = GetItemWidget(ItemBodyBuilder))
+		{
+			ParentBox->AddSlot()
+				.AutoHeight()
+				.Padding(0.f)
+				[
+					BodyWidget.ToSharedRef()
+				];
+		}
+
+		ParentBox->AddSlot()
+			.FillHeight(1.f)
+			.Padding(0.f)
+			[
+				SAssignNew(ItemsListView, SListView<FOperatorStackEditorItemPtr>)
+				.ListViewStyle(&FOperatorStackEditorStyle::Get().GetWidgetStyle<FTableViewStyle>("ListViewStyle"))
+				.ListItemsSource(&Items)
+				.ClearSelectionOnClick(true)
+				.SelectionMode(ESelectionMode::Multi)
+				.OnKeyDownHandler(this, &SOperatorStackEditorStack::OnKeyDownHandler)
+				.OnSelectionChanged(this, &SOperatorStackEditorStack::OnSelectionChanged)
+				.OnGenerateRow(this, &SOperatorStackEditorStack::OnGenerateRow)
+			];
 
 		return SNew(SBox)
 			.HAlign(HAlign_Fill)
 			.VAlign(VAlign_Fill)
 			.Padding(0.f)
 			[
-				ItemsListView.ToSharedRef()
+				ParentBox
 			];
 	}
 
@@ -537,39 +596,7 @@ TSharedPtr<SWidget> SOperatorStackEditorStack::GenerateBodyWidget()
 			EmptyBodyDefaultText = ItemBodyBuilder.GetEmptyBodyText();
 		}
 
-		// We have set a custom widget for this body
-		if (TSharedPtr<SWidget> CustomWidget = ItemBodyBuilder.GetCustomWidget())
-		{
-			return CustomWidget;
-		}
-
-		// Only build details view if we allow it
-		if (ItemBodyBuilder.GetShowDetailsView())
-		{
-			FCustomDetailsViewArgs BodyCustomViewArgs;
-			BodyCustomViewArgs.bShowCategories = false;
-			BodyCustomViewArgs.KeyframeHandler = MainPanel->GetKeyframeHandler();
-			BodyCustomViewArgs.bAllowGlobalExtensions = true;
-			BodyCustomViewArgs.ColumnSizeData = MainPanel->GetDetailColumnSize();
-
-			for (const TSharedPtr<FCustomDetailsViewItemId>& DetailsViewId : ItemBodyBuilder.GetDisallowedDetailsViewItems())
-			{
-				BodyCustomViewArgs.ItemAllowList.Disallow(*DetailsViewId.Get());
-			}
-
-			for (const TSharedPtr<FCustomDetailsViewItemId>& DetailsViewId : ItemBodyBuilder.GetAllowedDetailsViewItems())
-			{
-				BodyCustomViewArgs.ItemAllowList.Allow(*DetailsViewId.Get());
-			}
-
-			const FOperatorStackEditorItem* DetailViewItem = ItemBodyBuilder.GetDetailsViewItem().IsValid()
-				? ItemBodyBuilder.GetDetailsViewItem().Get()
-				: CustomizeItem.Get();
-
-			BodyDetailsView = CreateDetailsView(BodyCustomViewArgs, *DetailViewItem);
-
-			return BodyDetailsView;
-		}
+		return GetItemWidget(ItemBodyBuilder);
 	}
 
 	// No children and no items, display empty body info text
