MotionDesign : PropertyAnimator
- Support expand property and its children in custom details view
- Linked properties array from the animator will be fully expanded by default

diff --git a/Engine/Plugins/Experimental/Avalanche/Source/AvalancheSequencer/Private/DetailsView/SAvaMarkDetails.cpp b/Engine/Plugins/Experimental/Avalanche/Source/AvalancheSequencer/Private/DetailsView/SAvaMarkDetails.cpp
--- a/Engine/Plugins/Experimental/Avalanche/Source/AvalancheSequencer/Private/DetailsView/SAvaMarkDetails.cpp
+++ b/Engine/Plugins/Experimental/Avalanche/Source/AvalancheSequencer/Private/DetailsView/SAvaMarkDetails.cpp
@@ -22,8 +22,8 @@ void SAvaMarkDetails::Construct(const FArguments& InArgs, UAvaSequence* InSequen
 	CustomDetailsViewArgs.bAllowGlobalExtensions = true;
 	CustomDetailsViewArgs.CategoryAllowList.Allow(TEXT("Marks"));
 	CustomDetailsViewArgs.ItemAllowList.Disallow(FCustomDetailsViewItemId::MakePropertyId(TEXT("Label")));
-	CustomDetailsViewArgs.ExpansionState.Add(FCustomDetailsViewItemId::MakePropertyId<UAvaSequence>(TEXT("MotionDesign")), true);
-	CustomDetailsViewArgs.ExpansionState.Add(FCustomDetailsViewItemId::MakePropertyId<UAvaSequence>(TEXT("Marks")), true);
+	CustomDetailsViewArgs.ExpansionState.Add(FCustomDetailsViewItemId::MakePropertyId<UAvaSequence>(TEXT("MotionDesign")), ECustomDetailsViewExpansion::SelfExpanded);
+	CustomDetailsViewArgs.ExpansionState.Add(FCustomDetailsViewItemId::MakePropertyId<UAvaSequence>(TEXT("Marks")), ECustomDetailsViewExpansion::SelfExpanded);
 
 	TSharedRef<ICustomDetailsView> SettingsDetailsView = ICustomDetailsViewModule::Get().CreateCustomDetailsView(CustomDetailsViewArgs);
 
diff --git a/Engine/Plugins/Experimental/Avalanche/Source/AvalancheSequencer/Private/DetailsView/Section/AvaSequencePlaybackDetails.cpp b/Engine/Plugins/Experimental/Avalanche/Source/AvalancheSequencer/Private/DetailsView/Section/AvaSequencePlaybackDetails.cpp
--- a/Engine/Plugins/Experimental/Avalanche/Source/AvalancheSequencer/Private/DetailsView/Section/AvaSequencePlaybackDetails.cpp
+++ b/Engine/Plugins/Experimental/Avalanche/Source/AvalancheSequencer/Private/DetailsView/Section/AvaSequencePlaybackDetails.cpp
@@ -48,7 +48,7 @@ TSharedRef<SWidget> FAvaSequencePlaybackDetails::CreateContentWidget()
 	CustomDetailsViewArgs.bShowCategories = true;
 	CustomDetailsViewArgs.bAllowGlobalExtensions = true;
 	CustomDetailsViewArgs.CategoryAllowList.Allow(TEXT("Scheduled Playback"));
-	CustomDetailsViewArgs.ExpansionState.Add(FCustomDetailsViewItemId::MakeCategoryId("Scheduled Playback"), true);
+	CustomDetailsViewArgs.ExpansionState.Add(FCustomDetailsViewItemId::MakeCategoryId("Scheduled Playback"), ECustomDetailsViewExpansion::SelfExpanded);
 
 	const TSharedRef<ICustomDetailsView> PlaybackDetailsView = ICustomDetailsViewModule::Get().CreateCustomDetailsView(CustomDetailsViewArgs);
 
diff --git a/Engine/Plugins/Experimental/Avalanche/Source/AvalancheSequencer/Private/DetailsView/Section/AvaSequenceSettingsDetails.cpp b/Engine/Plugins/Experimental/Avalanche/Source/AvalancheSequencer/Private/DetailsView/Section/AvaSequenceSettingsDetails.cpp
--- a/Engine/Plugins/Experimental/Avalanche/Source/AvalancheSequencer/Private/DetailsView/Section/AvaSequenceSettingsDetails.cpp
+++ b/Engine/Plugins/Experimental/Avalanche/Source/AvalancheSequencer/Private/DetailsView/Section/AvaSequenceSettingsDetails.cpp
@@ -47,8 +47,8 @@ TSharedRef<SWidget> FAvaSequenceSettingsDetails::CreateContentWidget()
 	CustomDetailsViewArgs.bShowCategories = true;
 	CustomDetailsViewArgs.bAllowGlobalExtensions = true;
 	CustomDetailsViewArgs.CategoryAllowList.Allow(TEXT("Sequence Settings"));
-	CustomDetailsViewArgs.ExpansionState.Add(FCustomDetailsViewItemId::MakeCategoryId(TEXT("Sequence Settings")), true);
-	CustomDetailsViewArgs.ExpansionState.Add(FCustomDetailsViewItemId::MakePropertyId<UAvaSequence>(TEXT("Marks")), true);
+	CustomDetailsViewArgs.ExpansionState.Add(FCustomDetailsViewItemId::MakeCategoryId("Sequence Settings"), ECustomDetailsViewExpansion::SelfExpanded);
+	CustomDetailsViewArgs.ExpansionState.Add(FCustomDetailsViewItemId::MakePropertyId<UAvaSequence>(TEXT("Marks")), ECustomDetailsViewExpansion::SelfExpanded);
 
 	SettingsDetailsView = ICustomDetailsViewModule::Get().CreateCustomDetailsView(CustomDetailsViewArgs);
 
diff --git a/Engine/Plugins/Experimental/CustomDetailsView/Source/CustomDetailsView/Private/SCustomDetailsView.cpp b/Engine/Plugins/Experimental/CustomDetailsView/Source/CustomDetailsView/Private/SCustomDetailsView.cpp
--- a/Engine/Plugins/Experimental/CustomDetailsView/Source/CustomDetailsView/Private/SCustomDetailsView.cpp
+++ b/Engine/Plugins/Experimental/CustomDetailsView/Source/CustomDetailsView/Private/SCustomDetailsView.cpp
@@ -141,7 +141,25 @@ void SCustomDetailsView::OnExpansionChanged(TSharedPtr<ICustomDetailsViewItem> I
 	{
 		return;
 	}
-	ViewArgs.ExpansionState.Add(InItem->GetItemId(), bInExpanded);
+
+	if (const ECustomDetailsViewExpansion* Expansion = ViewArgs.ExpansionState.Find(InItem->GetItemId()))
+	{
+		const bool bExpanded = *Expansion != ECustomDetailsViewExpansion::Collapsed;
+
+		if (bExpanded)
+		{
+			ViewArgs.ExpansionState.Add(InItem->GetItemId(), bInExpanded ? *Expansion : ECustomDetailsViewExpansion::Collapsed);
+		}
+		else
+		{
+			ViewArgs.ExpansionState.Add(InItem->GetItemId(), bInExpanded ? ECustomDetailsViewExpansion::SelfExpanded : ECustomDetailsViewExpansion::Collapsed);
+		}
+	}
+	else
+	{
+		ViewArgs.ExpansionState.Add(InItem->GetItemId(), bInExpanded ? ECustomDetailsViewExpansion::SelfExpanded : ECustomDetailsViewExpansion::Collapsed);
+	}
+
 	ViewArgs.OnExpansionStateChanged.Broadcast(InItem.ToSharedRef(), bInExpanded);
 }
 
@@ -169,10 +187,33 @@ bool SCustomDetailsView::ShouldItemExpand(const TSharedPtr<ICustomDetailsViewIte
 		return false;
 	}
 
-	if (const bool* const FoundExpansionState = ViewArgs.ExpansionState.Find(InItem->GetItemId()))
+	TSharedPtr<ICustomDetailsViewItem> CheckItem = InItem;
+
+	do
 	{
-		return *FoundExpansionState;
+		if (const ECustomDetailsViewExpansion* const FoundExpansionState = ViewArgs.ExpansionState.Find(CheckItem->GetItemId()))
+		{
+			if (CheckItem.Get() == InItem.Get())
+			{
+				if (*FoundExpansionState == ECustomDetailsViewExpansion::Collapsed)
+				{
+					return false;
+				}
+
+				return true;
+			}
+
+			if (*FoundExpansionState == ECustomDetailsViewExpansion::SelfAndChildrenExpanded)
+			{
+				return true;
+			}
+
+			return false;
+		}
+
+		CheckItem = CheckItem->GetParent();
 	}
+	while (CheckItem.IsValid() && CheckItem->GetParent());
 
 	return ViewArgs.bDefaultItemsExpanded;
 }
@@ -324,20 +365,20 @@ bool SCustomDetailsView::FilterItems(const TArray<FString>& InFilterStrings)
 	return false;
 }
 
-bool SCustomDetailsView::GetItemExpansionState(const FCustomDetailsViewItemId& InItemId, bool& bOutExpanded) const
+bool SCustomDetailsView::GetItemExpansionState(const FCustomDetailsViewItemId& InItemId, ECustomDetailsViewExpansion& OutExpansion) const
 {
-	if (const bool* State = ViewArgs.ExpansionState.Find(InItemId))
+	if (const ECustomDetailsViewExpansion* State = ViewArgs.ExpansionState.Find(InItemId))
 	{
-		bOutExpanded = *State;
+		OutExpansion = *State;
 		return true;
 	}
 
 	return false;
 }
 
-void SCustomDetailsView::SetItemExpansionState(const FCustomDetailsViewItemId& InItemId, bool bInExpanded)
+void SCustomDetailsView::SetItemExpansionState(const FCustomDetailsViewItemId& InItemId, ECustomDetailsViewExpansion InExpansion)
 {
-	ViewArgs.ExpansionState.FindOrAdd(InItemId) = bInExpanded;
+	ViewArgs.ExpansionState.FindOrAdd(InItemId) = InExpansion;
 }
 
 bool SCustomDetailsView::ShouldRebuildImmediately(ECustomDetailsViewBuildType InBuildType) const
diff --git a/Engine/Plugins/Experimental/CustomDetailsView/Source/CustomDetailsView/Private/SCustomDetailsView.h b/Engine/Plugins/Experimental/CustomDetailsView/Source/CustomDetailsView/Private/SCustomDetailsView.h
--- a/Engine/Plugins/Experimental/CustomDetailsView/Source/CustomDetailsView/Private/SCustomDetailsView.h
+++ b/Engine/Plugins/Experimental/CustomDetailsView/Source/CustomDetailsView/Private/SCustomDetailsView.h
@@ -83,8 +83,8 @@ public:
 	virtual TSharedPtr<ICustomDetailsViewCustomCategoryItem> CreateCustomCategoryItem(FName InItemName, const FText& InLabel = FText::GetEmpty(), const FText& InToolTip = FText::GetEmpty()) override;
 	virtual TSharedPtr<ICustomDetailsViewItem> FindCustomItem(const FName& InItemName) const override;
 	virtual bool FilterItems(const TArray<FString>& InFilterStrings) override;
-	virtual bool GetItemExpansionState(const FCustomDetailsViewItemId& InItemId, bool& bOutExpanded) const override;
-	virtual void SetItemExpansionState(const FCustomDetailsViewItemId& InItemId, bool bInExpanded) override;
+	virtual bool GetItemExpansionState(const FCustomDetailsViewItemId& InItemId, ECustomDetailsViewExpansion& OutExpansion) const override;
+	virtual void SetItemExpansionState(const FCustomDetailsViewItemId& InItemId, ECustomDetailsViewExpansion InExpansion) override;
 	//~ End ICustomDetailsView
 
 private:
diff --git a/Engine/Plugins/Experimental/CustomDetailsView/Source/CustomDetailsView/Public/CustomDetailsViewArgs.h b/Engine/Plugins/Experimental/CustomDetailsView/Source/CustomDetailsView/Public/CustomDetailsViewArgs.h
--- a/Engine/Plugins/Experimental/CustomDetailsView/Source/CustomDetailsView/Public/CustomDetailsViewArgs.h
+++ b/Engine/Plugins/Experimental/CustomDetailsView/Source/CustomDetailsView/Public/CustomDetailsViewArgs.h
@@ -89,6 +89,14 @@ enum class ECustomDetailsViewWidgetType
 	Extensions,
 };
 
+/** Expansion state allowed for an item */
+enum class ECustomDetailsViewExpansion : uint8
+{
+	Collapsed,
+	SelfExpanded,
+	SelfAndChildrenExpanded
+};
+
 struct FCustomDetailsViewArgs
 {
 	/** List of Allowed and Disallowed Categories */
@@ -104,7 +112,7 @@ struct FCustomDetailsViewArgs
 	bool bDefaultItemsExpanded = false;
 
 	/** Map of the Node Name to their Expanded State */
-	TMap<FCustomDetailsViewItemId, bool> ExpansionState;
+	TMap<FCustomDetailsViewItemId, ECustomDetailsViewExpansion> ExpansionState;
 
 	DECLARE_MULTICAST_DELEGATE_TwoParams(FOnExpansionStateChanged, const TSharedRef<ICustomDetailsViewItem>&, bool)
 	FOnExpansionStateChanged OnExpansionStateChanged;
diff --git a/Engine/Plugins/Experimental/CustomDetailsView/Source/CustomDetailsView/Public/ICustomDetailsView.h b/Engine/Plugins/Experimental/CustomDetailsView/Source/CustomDetailsView/Public/ICustomDetailsView.h
--- a/Engine/Plugins/Experimental/CustomDetailsView/Source/CustomDetailsView/Public/ICustomDetailsView.h
+++ b/Engine/Plugins/Experimental/CustomDetailsView/Source/CustomDetailsView/Public/ICustomDetailsView.h
@@ -13,6 +13,7 @@ class ICustomDetailsViewCustomItem;
 class IDetailTreeNode;
 class UObject;
 template<typename ItemType> class STreeView;
+enum class ECustomDetailsViewExpansion : uint8;
 
 enum class ECustomDetailsViewBuildType : uint8
 {
@@ -102,8 +103,8 @@ public:
 	virtual TSharedPtr<ICustomDetailsViewItem> FindCustomItem(const FName& InItemName) const = 0;
 
 	/** Returns the current expansion state of an item. Returns true if the item state was found. */
-	virtual bool GetItemExpansionState(const FCustomDetailsViewItemId& InItemId, bool& bOutExpanded) const = 0;
+	virtual bool GetItemExpansionState(const FCustomDetailsViewItemId& InItemId, ECustomDetailsViewExpansion& OutExpansion) const = 0;
 
 	/** Sets the current expansion state of an item. */
-	virtual void SetItemExpansionState(const FCustomDetailsViewItemId& InItemId, bool bInExpanded) = 0;
+	virtual void SetItemExpansionState(const FCustomDetailsViewItemId& InItemId, ECustomDetailsViewExpansion InExpansion) = 0;
 };
diff --git a/Engine/Plugins/Experimental/DynamicMaterial/Source/DynamicMaterialEditor/Private/Slate/SDMComponentEdit.cpp b/Engine/Plugins/Experimental/DynamicMaterial/Source/DynamicMaterialEditor/Private/Slate/SDMComponentEdit.cpp
--- a/Engine/Plugins/Experimental/DynamicMaterial/Source/DynamicMaterialEditor/Private/Slate/SDMComponentEdit.cpp
+++ b/Engine/Plugins/Experimental/DynamicMaterial/Source/DynamicMaterialEditor/Private/Slate/SDMComponentEdit.cpp
@@ -335,7 +335,7 @@ TSharedRef<SWidget> SDMComponentEdit::CreateEditWidget()
 				bool bExpansionState = true;
 				SDMEditor::GetExpansionState(Component, DefaultCategoryName, bExpansionState);
 
-				DetailsView->SetItemExpansionState(DefaultCategoryItem->GetItemId(), bExpansionState);
+				DetailsView->SetItemExpansionState(DefaultCategoryItem->GetItemId(), bExpansionState ? ECustomDetailsViewExpansion::SelfExpanded : ECustomDetailsViewExpansion::Collapsed);
 				Categories.Add(DefaultCategoryName);
 			}
 
@@ -432,7 +432,7 @@ TSharedRef<SWidget> SDMComponentEdit::CreateEditWidget()
 				bool bExpansionState = true;
 				SDMEditor::GetExpansionState(Component, CategoryName, bExpansionState);
 
-				DetailsView->SetItemExpansionState(CategoryItem->GetItemId(), bExpansionState);
+				DetailsView->SetItemExpansionState(CategoryItem->GetItemId(), bExpansionState ? ECustomDetailsViewExpansion::SelfExpanded : ECustomDetailsViewExpansion::Collapsed);
 
 				Categories.Add(CategoryName);
 			}
diff --git a/Engine/Plugins/Experimental/OperatorStack/Source/OperatorStack/Private/Builders/OperatorStackEditorBodyBuilder.cpp b/Engine/Plugins/Experimental/OperatorStack/Source/OperatorStack/Private/Builders/OperatorStackEditorBodyBuilder.cpp
--- a/Engine/Plugins/Experimental/OperatorStack/Source/OperatorStack/Private/Builders/OperatorStackEditorBodyBuilder.cpp
+++ b/Engine/Plugins/Experimental/OperatorStack/Source/OperatorStack/Private/Builders/OperatorStackEditorBodyBuilder.cpp
@@ -45,6 +45,18 @@ FOperatorStackEditorBodyBuilder& FOperatorStackEditorBodyBuilder::AllowCategory(
 	return *this;
 }
 
+FOperatorStackEditorBodyBuilder& FOperatorStackEditorBodyBuilder::ExpandProperty(FProperty* InProperty)
+{
+	ExpandedDetailsViewItems.Add(MakeShared<FCustomDetailsViewItemId>(FCustomDetailsViewItemId::MakePropertyId(InProperty)));
+	return *this;
+}
+
+FOperatorStackEditorBodyBuilder& FOperatorStackEditorBodyBuilder::CollapseProperty(FProperty* InProperty)
+{
+	CollapsedDetailsViewItems.Add(MakeShared<FCustomDetailsViewItemId>(FCustomDetailsViewItemId::MakePropertyId(InProperty)));
+	return *this;
+}
+
 FOperatorStackEditorBodyBuilder& FOperatorStackEditorBodyBuilder::SetEmptyBodyText(const FText& InText)
 {
 	EmptyBodyText = InText;
diff --git a/Engine/Plugins/Experimental/OperatorStack/Source/OperatorStack/Private/Widgets/SOperatorStackEditorStack.cpp b/Engine/Plugins/Experimental/OperatorStack/Source/OperatorStack/Private/Widgets/SOperatorStackEditorStack.cpp
--- a/Engine/Plugins/Experimental/OperatorStack/Source/OperatorStack/Private/Widgets/SOperatorStackEditorStack.cpp
+++ b/Engine/Plugins/Experimental/OperatorStack/Source/OperatorStack/Private/Widgets/SOperatorStackEditorStack.cpp
@@ -481,21 +481,31 @@ TSharedPtr<SWidget> SOperatorStackEditorStack::GenerateBodyWidget()
 			BodyCustomViewArgs.bAllowGlobalExtensions = true;
 			BodyCustomViewArgs.ColumnSizeData = MainPanel->GetDetailColumnSize();
 
-			for (const TSharedPtr<FCustomDetailsViewItemId>& DetailsViewId : InItem.GetDisallowedDetailsViewItems())
+			for (const TSharedRef<FCustomDetailsViewItemId>& DetailsViewId : InItem.GetDisallowedDetailsViewItems())
 			{
-				BodyCustomViewArgs.ItemAllowList.Disallow(*DetailsViewId.Get());
+				BodyCustomViewArgs.ItemAllowList.Disallow(DetailsViewId.Get());
 			}
 
-			for (const TSharedPtr<FCustomDetailsViewItemId>& DetailsViewId : InItem.GetAllowedDetailsViewItems())
+			for (const TSharedRef<FCustomDetailsViewItemId>& DetailsViewId : InItem.GetAllowedDetailsViewItems())
 			{
-				BodyCustomViewArgs.ItemAllowList.Allow(*DetailsViewId.Get());
+				BodyCustomViewArgs.ItemAllowList.Allow(DetailsViewId.Get());
+			}
+
+			for (const TSharedRef<FCustomDetailsViewItemId>& DetailsViewId : InItem.GetCollapsedDetailsViewItems())
+			{
+				BodyCustomViewArgs.ExpansionState.Add(DetailsViewId.Get(), ECustomDetailsViewExpansion::Collapsed);
+			}
+
+			for (const TSharedRef<FCustomDetailsViewItemId>& DetailsViewId : InItem.GetExpandedDetailsViewItems())
+			{
+				BodyCustomViewArgs.ExpansionState.Add(DetailsViewId.Get(), ECustomDetailsViewExpansion::SelfAndChildrenExpanded);
 			}
 
 			const FOperatorStackEditorItem* DetailViewItem = InItem.GetDetailsViewItem().IsValid()
 				? InItem.GetDetailsViewItem().Get()
 				: CustomizeItem.Get();
 
-			BodyDetailsView = CreateDetailsView(BodyCustomViewArgs, *DetailViewItem); 
+			BodyDetailsView = CreateDetailsView(BodyCustomViewArgs, *DetailViewItem);
 
 			return BodyDetailsView;
 		}
@@ -642,14 +652,24 @@ TSharedPtr<SWidget> SOperatorStackEditorStack::GenerateFooterWidget()
 		FooterCustomViewArgs.bAllowGlobalExtensions = true;
 		FooterCustomViewArgs.ColumnSizeData = MainPanel->GetDetailColumnSize();
 
-		for (const TSharedPtr<FCustomDetailsViewItemId>& DetailsViewId : FooterBuilder.GetDisallowedDetailsViewItems())
+		for (const TSharedRef<FCustomDetailsViewItemId>& DetailsViewId : FooterBuilder.GetDisallowedDetailsViewItems())
+		{
+			FooterCustomViewArgs.ItemAllowList.Disallow(DetailsViewId.Get());
+		}
+
+		for (const TSharedRef<FCustomDetailsViewItemId>& DetailsViewId : FooterBuilder.GetAllowedDetailsViewItems())
+		{
+			FooterCustomViewArgs.ItemAllowList.Allow(DetailsViewId.Get());
+		}
+
+		for (const TSharedRef<FCustomDetailsViewItemId>& DetailsViewId : FooterBuilder.GetCollapsedDetailsViewItems())
 		{
-			FooterCustomViewArgs.ItemAllowList.Disallow(*DetailsViewId.Get());
+			FooterCustomViewArgs.ExpansionState.Add(DetailsViewId.Get(), ECustomDetailsViewExpansion::Collapsed);
 		}
 
-		for (const TSharedPtr<FCustomDetailsViewItemId>& DetailsViewId : FooterBuilder.GetAllowedDetailsViewItems())
+		for (const TSharedRef<FCustomDetailsViewItemId>& DetailsViewId : FooterBuilder.GetExpandedDetailsViewItems())
 		{
-			FooterCustomViewArgs.ItemAllowList.Allow(*DetailsViewId.Get());
+			FooterCustomViewArgs.ExpansionState.Add(DetailsViewId.Get(), ECustomDetailsViewExpansion::SelfAndChildrenExpanded);
 		}
 
 		const FOperatorStackEditorItem* DetailViewItem = FooterBuilder.GetDetailsViewItem().IsValid()
diff --git a/Engine/Plugins/Experimental/OperatorStack/Source/OperatorStack/Public/Builders/OperatorStackEditorBodyBuilder.h b/Engine/Plugins/Experimental/OperatorStack/Source/OperatorStack/Public/Builders/OperatorStackEditorBodyBuilder.h
--- a/Engine/Plugins/Experimental/OperatorStack/Source/OperatorStack/Public/Builders/OperatorStackEditorBodyBuilder.h
+++ b/Engine/Plugins/Experimental/OperatorStack/Source/OperatorStack/Public/Builders/OperatorStackEditorBodyBuilder.h
@@ -14,12 +14,13 @@ struct OPERATORSTACKEDITOR_API FOperatorStackEditorBodyBuilder
 {
 	/** Override the stack generated widget with a custom one */
 	FOperatorStackEditorBodyBuilder& SetCustomWidget(TSharedPtr<SWidget> InWidget);
-	
+
 	/** Show a details view for the current item */
 	FOperatorStackEditorBodyBuilder& SetShowDetailsView(bool bInShowDetailsView);
 
+	/** Set the item that needs to be displayed by the details view overriding current item */
 	FOperatorStackEditorBodyBuilder& SetDetailsViewItem(TSharedPtr<FOperatorStackEditorItem> InItem);
-	
+
 	/** Disallow specific property in the details view */
 	FOperatorStackEditorBodyBuilder& DisallowProperty(FProperty* InProperty);
 
@@ -32,6 +33,12 @@ struct OPERATORSTACKEDITOR_API FOperatorStackEditorBodyBuilder
 	/** Allow specific category in the details view */
 	FOperatorStackEditorBodyBuilder& AllowCategory(const FName& InCategory);
 
+	/** Expand property state in the details view */
+	FOperatorStackEditorBodyBuilder& ExpandProperty(FProperty* InProperty);
+
+	/** Collapse property state in the details view */
+	FOperatorStackEditorBodyBuilder& CollapseProperty(FProperty* InProperty);
+
 	/** Set a text that will be displayed when it is empty */
 	FOperatorStackEditorBodyBuilder& SetEmptyBodyText(const FText& InText);
 
@@ -44,22 +51,32 @@ struct OPERATORSTACKEDITOR_API FOperatorStackEditorBodyBuilder
 	{
 		return bShowDetailsView;
 	}
-	
+
 	TSharedPtr<FOperatorStackEditorItem> GetDetailsViewItem() const
 	{
-		return DetailsViewItem;	
+		return DetailsViewItem;
 	}
-	
-	const TArray<TSharedPtr<FCustomDetailsViewItemId>>& GetDisallowedDetailsViewItems() const
+
+	const TArray<TSharedRef<FCustomDetailsViewItemId>>& GetDisallowedDetailsViewItems() const
 	{
 		return DisallowedDetailsViewItems;
 	}
 
-	const TArray<TSharedPtr<FCustomDetailsViewItemId>>& GetAllowedDetailsViewItems() const
+	const TArray<TSharedRef<FCustomDetailsViewItemId>>& GetAllowedDetailsViewItems() const
 	{
 		return AllowedDetailsViewItems;
 	}
-	
+
+	const TArray<TSharedRef<FCustomDetailsViewItemId>>& GetCollapsedDetailsViewItems() const
+	{
+		return CollapsedDetailsViewItems;
+	}
+
+	const TArray<TSharedRef<FCustomDetailsViewItemId>>& GetExpandedDetailsViewItems() const
+	{
+		return ExpandedDetailsViewItems;
+	}
+
 	const FText& GetEmptyBodyText() const
 	{
 		return EmptyBodyText;
@@ -68,10 +85,10 @@ struct OPERATORSTACKEDITOR_API FOperatorStackEditorBodyBuilder
 protected:
 	/** Custom widget to replace content */
 	TSharedPtr<SWidget> CustomWidget = nullptr;
-	
+
 	/** Does this body contains a details view */
 	bool bShowDetailsView = false;
-	
+
 	/** Override actual item to display detail view */
 	TSharedPtr<FOperatorStackEditorItem> DetailsViewItem = nullptr;
 
@@ -79,8 +96,14 @@ protected:
 	FText EmptyBodyText;
 
 	/** Disallowed items inside details view */
-	TArray<TSharedPtr<FCustomDetailsViewItemId>> DisallowedDetailsViewItems;
-	
+	TArray<TSharedRef<FCustomDetailsViewItemId>> DisallowedDetailsViewItems;
+
 	/** Allowed items inside details view */
-	TArray<TSharedPtr<FCustomDetailsViewItemId>> AllowedDetailsViewItems;
+	TArray<TSharedRef<FCustomDetailsViewItemId>> AllowedDetailsViewItems;
+
+	/** Collapsed items inside details view */
+	TArray<TSharedRef<FCustomDetailsViewItemId>> CollapsedDetailsViewItems;
+
+	/** Expanded items inside details view */
+	TArray<TSharedRef<FCustomDetailsViewItemId>> ExpandedDetailsViewItems;
 };
\ No newline at end of file
diff --git a/Engine/Plugins/Experimental/PropertyAnimatorCore/Source/PropertyAnimatorCore/Private/Animators/PropertyAnimatorCoreBase.cpp b/Engine/Plugins/Experimental/PropertyAnimatorCore/Source/PropertyAnimatorCore/Private/Animators/PropertyAnimatorCoreBase.cpp
--- a/Engine/Plugins/Experimental/PropertyAnimatorCore/Source/PropertyAnimatorCore/Private/Animators/PropertyAnimatorCoreBase.cpp
+++ b/Engine/Plugins/Experimental/PropertyAnimatorCore/Source/PropertyAnimatorCore/Private/Animators/PropertyAnimatorCoreBase.cpp
@@ -18,6 +18,13 @@ UPropertyAnimatorCoreBase::FOnAnimatorUpdated UPropertyAnimatorCoreBase::OnAnima
 UPropertyAnimatorCoreBase::FOnAnimatorPropertyUpdated UPropertyAnimatorCoreBase::OnAnimatorPropertyLinkedDelegate;
 UPropertyAnimatorCoreBase::FOnAnimatorPropertyUpdated UPropertyAnimatorCoreBase::OnAnimatorPropertyUnlinkedDelegate;
 
+#if WITH_EDITOR
+FName UPropertyAnimatorCoreBase::GetLinkedPropertiesPropertyName()
+{
+	return GET_MEMBER_NAME_CHECKED(UPropertyAnimatorCoreBase, LinkedProperties);
+}
+#endif
+
 UPropertyAnimatorCoreBase::UPropertyAnimatorCoreBase()
 {
 	if (const UPropertyAnimatorCoreSubsystem* AnimatorSubsystem = UPropertyAnimatorCoreSubsystem::Get())
@@ -40,16 +47,6 @@ UPropertyAnimatorCoreComponent* UPropertyAnimatorCoreBase::GetAnimatorComponent(
 	return GetTypedOuter<UPropertyAnimatorCoreComponent>();
 }
 
-float UPropertyAnimatorCoreBase::GetAnimatorComponentMagnitude() const
-{
-	if (const UPropertyAnimatorCoreComponent* AnimatorComponent = GetAnimatorComponent())
-	{
-		return AnimatorComponent->GetAnimatorsMagnitude();
-	}
-
-	return 0.f;
-}
-
 void UPropertyAnimatorCoreBase::UpdateAnimatorDisplayName()
 {
 	TArray<FString> PropertiesNames;
diff --git a/Engine/Plugins/Experimental/PropertyAnimatorCore/Source/PropertyAnimatorCore/Public/Animators/PropertyAnimatorCoreBase.h b/Engine/Plugins/Experimental/PropertyAnimatorCore/Source/PropertyAnimatorCore/Public/Animators/PropertyAnimatorCoreBase.h
--- a/Engine/Plugins/Experimental/PropertyAnimatorCore/Source/PropertyAnimatorCore/Public/Animators/PropertyAnimatorCoreBase.h
+++ b/Engine/Plugins/Experimental/PropertyAnimatorCore/Source/PropertyAnimatorCore/Public/Animators/PropertyAnimatorCoreBase.h
@@ -56,15 +56,16 @@ public:
 	static constexpr const TCHAR* FrequencyParameterName = TEXT("Frequency");
 	static constexpr const TCHAR* AlphaParameterName = TEXT("Alpha");
 
+#if WITH_EDITOR
+	PROPERTYANIMATORCORE_API static FName GetLinkedPropertiesPropertyName();
+#endif
+
 	PROPERTYANIMATORCORE_API UPropertyAnimatorCoreBase();
 
 	PROPERTYANIMATORCORE_API AActor* GetAnimatorActor() const;
 
 	PROPERTYANIMATORCORE_API UPropertyAnimatorCoreComponent* GetAnimatorComponent() const;
 
-	/** Get the global animators magnitude */
-	float GetAnimatorComponentMagnitude() const;
-
 	/** Set the state of this animator */
 	PROPERTYANIMATORCORE_API void SetAnimatorEnabled(bool bInIsEnabled);
 	bool GetAnimatorEnabled() const
diff --git a/Engine/Plugins/Experimental/PropertyAnimatorCore/Source/PropertyAnimatorCoreEditor/Private/Customizations/PropertyAnimatorCoreEditorStackCustomization.cpp b/Engine/Plugins/Experimental/PropertyAnimatorCore/Source/PropertyAnimatorCoreEditor/Private/Customizations/PropertyAnimatorCoreEditorStackCustomization.cpp
--- a/Engine/Plugins/Experimental/PropertyAnimatorCore/Source/PropertyAnimatorCoreEditor/Private/Customizations/PropertyAnimatorCoreEditorStackCustomization.cpp
+++ b/Engine/Plugins/Experimental/PropertyAnimatorCore/Source/PropertyAnimatorCoreEditor/Private/Customizations/PropertyAnimatorCoreEditorStackCustomization.cpp
@@ -295,11 +295,13 @@ void UPropertyAnimatorCoreEditorStackCustomization::CustomizeItemBody(const FOpe
 	{
 		FBoolProperty* EnableProperty = FindFProperty<FBoolProperty>(UPropertyAnimatorCoreBase::StaticClass(), GET_MEMBER_NAME_CHECKED(UPropertyAnimatorCoreBase, bAnimatorEnabled));
 		FProperty* TimeElapsedProperty = FindFProperty<FProperty>(UPropertyAnimatorCoreTimeSourceBase::StaticClass(), UPropertyAnimatorCoreTimeSourceBase::GetTimeElapsedPropertyName());
+		FProperty* LinkedPropertiesProperty = FindFProperty<FProperty>(UPropertyAnimatorCoreBase::StaticClass(), UPropertyAnimatorCoreBase::GetLinkedPropertiesPropertyName());
 
 		InBodyBuilder
 			.SetShowDetailsView(true)
 			.DisallowProperty(EnableProperty)
-			.DisallowProperty(TimeElapsedProperty);
+			.DisallowProperty(TimeElapsedProperty)
+			.ExpandProperty(LinkedPropertiesProperty);
 	}
 
 	Super::CustomizeItemBody(InItem, InItemTree, InBodyBuilder);
