MotionDesign : OperatorStack
- Keep last search active and filter items on change without clearing filters

diff --git a/Engine/Plugins/Experimental/OperatorStack/Source/OperatorStack/Private/Widgets/SOperatorStackEditorPanel.cpp b/Engine/Plugins/Experimental/OperatorStack/Source/OperatorStack/Private/Widgets/SOperatorStackEditorPanel.cpp
--- a/Engine/Plugins/Experimental/OperatorStack/Source/OperatorStack/Private/Widgets/SOperatorStackEditorPanel.cpp
+++ b/Engine/Plugins/Experimental/OperatorStack/Source/OperatorStack/Private/Widgets/SOperatorStackEditorPanel.cpp
@@ -2,10 +2,10 @@
 
 #include "Widgets/SOperatorStackEditorPanel.h"
 
-#include "DetailColumnSizeData.h"
+#include "Containers/Ticker.h"
 #include "Customizations/OperatorStackEditorStackCustomization.h"
+#include "DetailColumnSizeData.h"
 #include "SOperatorStackEditorStack.h"
-#include "Styling/ToolBarStyle.h"
 #include "Subsystems/OperatorStackEditorSubsystem.h"
 #include "Widgets/Images/SImage.h"
 #include "Widgets/Input/SCheckBox.h"
@@ -25,8 +25,6 @@ SOperatorStackEditorPanel::~SOperatorStackEditorPanel()
 
 void SOperatorStackEditorPanel::Construct(const FArguments& InArgs)
 {
-	static const FToolBarStyle& ToolBarStyle = FAppStyle::Get().GetWidgetStyle<FToolBarStyle>("SlimToolBar");
-
 	PanelId = InArgs._PanelId;
 
 	// Generate default empty context
@@ -273,6 +271,8 @@ void SOperatorStackEditorPanel::UpdateSlots()
 {
 	const UOperatorStackEditorSubsystem* EditorSubsystem = UOperatorStackEditorSubsystem::Get();
 
+	CustomizationStacks.Empty(NamedStackIndexes.Num());
+
 	for (const TPair<FName, int32>& NamedSlot : NamedStackIndexes)
 	{
 		if (const TSharedPtr<SBox> Box = StaticCastSharedPtr<SBox>(WidgetSwitcher->GetWidget(NamedSlot.Value)))
@@ -283,17 +283,17 @@ void SOperatorStackEditorPanel::UpdateSlots()
 
 			CustomizationTrees.Add(Customization, FOperatorStackEditorTree(Customization, Context));
 
+			TSharedRef<SOperatorStackEditorStack> Stack = SNew(SOperatorStackEditorStack, SharedThis(this), Customization, nullptr);
+			CustomizationStacks.Add(Stack);
+
 			Box->SetContent(
-				SNew(SOperatorStackEditorStack
-					, SharedThis(this)
-					, Customization
-					, nullptr)
+				Stack
 			);
 		}
 	}
 }
 
-void SOperatorStackEditorPanel::OnToolbarButtonClicked(ECheckBoxState InState, int32 InWidgetIdx) const
+void SOperatorStackEditorPanel::OnToolbarButtonClicked(ECheckBoxState InState, int32 InWidgetIdx)
 {
 	if (!WidgetSwitcher.IsValid() || InWidgetIdx == INDEX_NONE)
 	{
@@ -301,6 +301,9 @@ void SOperatorStackEditorPanel::OnToolbarButtonClicked(ECheckBoxState InState, i
 	}
 
 	WidgetSwitcher->SetActiveWidgetIndex(InWidgetIdx);
+
+	// Refresh search for new active item
+	FilterItemsAsync(LastSearch);
 }
 
 ECheckBoxState SOperatorStackEditorPanel::IsToolbarButtonActive(int32 InWidgetIdx) const
@@ -314,3 +317,29 @@ ECheckBoxState SOperatorStackEditorPanel::IsToolbarButtonActive(int32 InWidgetId
 		? ECheckBoxState::Checked
 		: ECheckBoxState::Unchecked;
 }
+
+void SOperatorStackEditorPanel::FilterItemsAsync(const FText& InNewSearch, float InDelay)
+{
+	LastSearch = InNewSearch;
+
+	// Cancel previous search
+	if (LastSearchHandle.IsValid())
+	{
+		FTSTicker::GetCoreTicker().RemoveTicker(LastSearchHandle);
+		LastSearchHandle.Reset();
+	}
+
+	// Start a new search on active items
+	LastSearchHandle = FTSTicker::GetCoreTicker().AddTicker(FTickerDelegate::CreateSPLambda(this, [this](float)->bool
+		{
+			const int32 Index = WidgetSwitcher->GetActiveWidgetIndex();
+
+			if (CustomizationStacks.IsValidIndex(Index))
+			{
+				CustomizationStacks[Index]->FilterItems(LastSearch);
+			}
+
+			return false;
+		}
+	), FMath::Abs(InDelay));
+}
diff --git a/Engine/Plugins/Experimental/OperatorStack/Source/OperatorStack/Private/Widgets/SOperatorStackEditorPanel.h b/Engine/Plugins/Experimental/OperatorStack/Source/OperatorStack/Private/Widgets/SOperatorStackEditorPanel.h
--- a/Engine/Plugins/Experimental/OperatorStack/Source/OperatorStack/Private/Widgets/SOperatorStackEditorPanel.h
+++ b/Engine/Plugins/Experimental/OperatorStack/Source/OperatorStack/Private/Widgets/SOperatorStackEditorPanel.h
@@ -79,6 +79,13 @@ public:
 		return DetailColumnSize;
 	}
 
+	const FText& GetLastSearch() const
+	{
+		return LastSearch;
+	}
+
+	void FilterItemsAsync(const FText& InNewSearch, float InDelay = 0.f);
+
 protected:
 	/** Is the toolbar button active */
 	EVisibility GetToolbarButtonVisibility(int32 InIdx) const;
@@ -90,7 +97,7 @@ protected:
 	void UpdateSlots();
 
 	/** Switch view when toolbar button is clicked */
-	void OnToolbarButtonClicked(ECheckBoxState InState, int32 InWidgetIdx) const;
+	void OnToolbarButtonClicked(ECheckBoxState InState, int32 InWidgetIdx);
 
 	/** Highlight active selection */
 	ECheckBoxState IsToolbarButtonActive(int32 InWidgetIdx) const;
@@ -128,4 +135,13 @@ private:
 
 	/** To check whether an item was expanded and restore it */
 	TMap<uint32, bool> ItemExpansionState;
+
+	/** Contains the customization stacks */
+	TArray<TSharedRef<SOperatorStackEditorStack>> CustomizationStacks;
+
+	/** Last search done within the stack */
+	FText LastSearch = FText::GetEmpty();
+
+	/** Last search request triggered */
+	FTSTicker::FDelegateHandle LastSearchHandle;
 };
\ No newline at end of file
diff --git a/Engine/Plugins/Experimental/OperatorStack/Source/OperatorStack/Private/Widgets/SOperatorStackEditorStack.cpp b/Engine/Plugins/Experimental/OperatorStack/Source/OperatorStack/Private/Widgets/SOperatorStackEditorStack.cpp
--- a/Engine/Plugins/Experimental/OperatorStack/Source/OperatorStack/Private/Widgets/SOperatorStackEditorStack.cpp
+++ b/Engine/Plugins/Experimental/OperatorStack/Source/OperatorStack/Private/Widgets/SOperatorStackEditorStack.cpp
@@ -5,22 +5,21 @@
 #include "Builders/OperatorStackEditorBodyBuilder.h"
 #include "Builders/OperatorStackEditorFooterBuilder.h"
 #include "Builders/OperatorStackEditorHeaderBuilder.h"
+#include "Contexts/OperatorStackEditorMenuContext.h"
 #include "CustomDetailsViewModule.h"
 #include "Customizations/OperatorStackEditorStackCustomization.h"
 #include "Framework/Application/SlateApplication.h"
 #include "ICustomDetailsView.h"
-#include "SOperatorStackEditorPanel.h"
 #include "Items/ICustomDetailsViewItem.h"
+#include "Items/OperatorStackEditorTree.h"
+#include "Items/OperatorStackEditorStructItem.h"
+#include "SOperatorStackEditorPanel.h"
 #include "SOperatorStackExpanderButton.h"
 #include "SPositiveActionButton.h"
 #include "Styles/OperatorStackEditorStyle.h"
 #include "Styling/AppStyle.h"
-#include "Styling/ToolBarStyle.h"
 #include "ToolMenu.h"
 #include "ToolMenus.h"
-#include "Contexts/OperatorStackEditorMenuContext.h"
-#include "Items/OperatorStackEditorTree.h"
-#include "Items/OperatorStackEditorStructItem.h"
 #include "Widgets/Colors/SColorBlock.h"
 #include "Widgets/Input/SCheckBox.h"
 #include "Widgets/Input/SComboButton.h"
@@ -78,6 +77,12 @@ void SOperatorStackEditorStack::Construct(const FArguments& InArgs
 	[
 		GenerateStackWidget()
 	];
+
+	// Refresh search when we hit a leaf
+	if (CustomizeItem.IsValid() && Items.IsEmpty())
+	{
+		InMainPanel->FilterItemsAsync(InMainPanel->GetLastSearch());
+	}
 }
 
 FOperatorStackEditorContextPtr SOperatorStackEditorStack::GetContext() const
@@ -295,6 +300,7 @@ TSharedPtr<SWidget> SOperatorStackEditorStack::GenerateHeaderWidget()
             .AutoHeight()
             [
             	SAssignNew(SearchBox, SSearchBox)
+	            .InitialText(MainPanel->GetLastSearch())
             	.HintText(LOCTEXT("OperatorStackEditorStackSearchHint", "Search items"))
             	.OnTextChanged(this, &SOperatorStackEditorStack::OnSearchTextChanged)
             	.OnTextCommitted(this, &SOperatorStackEditorStack::OnSearchTextCommitted)
@@ -741,31 +747,53 @@ EVisibility SOperatorStackEditorStack::GetMessageBoxIconVisibility() const
 	return GetMessageBoxIcon() ? EVisibility::Visible : EVisibility::Collapsed;
 }
 
+void SOperatorStackEditorStack::RequestSearchAsync() const
+{
+	if (const TSharedPtr<SOperatorStackEditorPanel> MainPanel = GetMainPanel())
+	{
+		if (SearchBox.IsValid())
+		{
+			MainPanel->FilterItemsAsync(SearchBox->GetText());
+		}
+	}
+}
+
 void SOperatorStackEditorStack::OnSearchTextChanged(const FText& InSearchText)
 {
-	OnSearchChanged();
+	RequestSearchAsync();
 }
 
-void SOperatorStackEditorStack::OnSearchChanged()
+void SOperatorStackEditorStack::FilterItems(const FText& InText)
 {
+	if (!SearchBox.IsValid())
+	{
+		return;
+	}
+
+	if (!SearchBox->GetText().EqualTo(InText))
+	{
+		SearchBox->SetText(InText);
+	}
+
 	const TSet<FString> SearchOR = SearchedKeywords;
 	TSet<FString> SearchAND;
 
-	const FString FilterString = SearchBox->GetText().ToString();
+	const FString FilterString = InText.ToString();
 	if (!FilterString.IsEmpty())
 	{
 		SearchAND.Add(FilterString);
 	}
 
-	HandleSearch(SearchOR, SearchAND);
+	HandleRecursiveSearch(SearchOR, SearchAND);
 }
 
 void SOperatorStackEditorStack::OnSearchTextCommitted(const FText& InFilterText, ETextCommit::Type InCommitType)
 {
 	if (InCommitType == ETextCommit::OnCleared)
 	{
-		SearchBox->SetText(FText::GetEmpty());
-		OnSearchChanged();
+		const FText& EmptyText = FText::GetEmpty();
+		SearchBox->SetText(EmptyText);
+		OnSearchTextChanged(EmptyText);
 		FSlateApplication::Get().ClearKeyboardFocus(EFocusCause::Cleared);
 	}
 }
@@ -781,7 +809,7 @@ void SOperatorStackEditorStack::OnSearchPinnedKeyword(ECheckBoxState InCheckStat
 		SearchedKeywords.Remove(InPinnedKeyword);
 	}
 
-	OnSearchChanged();
+	RequestSearchAsync();
 }
 
 bool SOperatorStackEditorStack::MatchSearch(const TSet<FString>& InSearchedKeywordsOR, const TSet<FString>& InSearchedKeywordsAND) const
@@ -861,7 +889,7 @@ bool SOperatorStackEditorStack::MatchSearch(const TSet<FString>& InSearchedKeywo
 	return bORSearchMatched && bANDSearchMatched;
 }
 
-bool SOperatorStackEditorStack::HandleSearch(const TSet<FString>& InSearchedKeywordsOR, const TSet<FString>& InSearchedKeywordsAND)
+bool SOperatorStackEditorStack::HandleRecursiveSearch(const TSet<FString>& InSearchedKeywords_OR, const TSet<FString>& InSearchedKeywords_AND)
 {
 	bool bMatchSearch = false;
 
@@ -869,7 +897,7 @@ bool SOperatorStackEditorStack::HandleSearch(const TSet<FString>& InSearchedKeyw
 	{
 		if (ItemWidget.IsValid())
 		{
-			if (ItemWidget->HandleSearch(InSearchedKeywordsOR, InSearchedKeywordsAND))
+			if (ItemWidget->HandleRecursiveSearch(InSearchedKeywords_OR, InSearchedKeywords_AND))
 			{
 				bMatchSearch = true;
 			}
@@ -881,15 +909,15 @@ bool SOperatorStackEditorStack::HandleSearch(const TSet<FString>& InSearchedKeyw
 
 	if (BodyDetailsView)
 	{
-		bRowFound |= BodyDetailsView->FilterItems(InSearchedKeywordsAND.Array());
+		bRowFound |= BodyDetailsView->FilterItems(InSearchedKeywords_AND.Array());
 	}
 
 	if (FooterDetailsView)
 	{
-		bRowFound |= FooterDetailsView->FilterItems(InSearchedKeywordsAND.Array());
+		bRowFound |= FooterDetailsView->FilterItems(InSearchedKeywords_AND.Array());
 	}
 
-	if (!InSearchedKeywordsAND.IsEmpty())
+	if (!InSearchedKeywords_AND.IsEmpty())
 	{
 		bMatchSearch |= bRowFound;
 	}
@@ -897,7 +925,7 @@ bool SOperatorStackEditorStack::HandleSearch(const TSet<FString>& InSearchedKeyw
 	// Do not hide root item if nothing was found
 	if (CustomizeItem.IsValid())
 	{
-		bMatchSearch |= MatchSearch(InSearchedKeywordsOR, InSearchedKeywordsAND);
+		bMatchSearch |= MatchSearch(InSearchedKeywords_OR, InSearchedKeywords_AND);
 		SetVisibility(bMatchSearch ? EVisibility::Visible : EVisibility::Collapsed);
 	}
 
diff --git a/Engine/Plugins/Experimental/OperatorStack/Source/OperatorStack/Private/Widgets/SOperatorStackEditorStack.h b/Engine/Plugins/Experimental/OperatorStack/Source/OperatorStack/Private/Widgets/SOperatorStackEditorStack.h
--- a/Engine/Plugins/Experimental/OperatorStack/Source/OperatorStack/Private/Widgets/SOperatorStackEditorStack.h
+++ b/Engine/Plugins/Experimental/OperatorStack/Source/OperatorStack/Private/Widgets/SOperatorStackEditorStack.h
@@ -50,6 +50,9 @@ public:
 
 	FOperatorStackEditorContextPtr GetContext() const;
 
+	/** Apply a filter on all items */
+	void FilterItems(const FText& InText);
+
 protected:
 	TSharedRef<ITableRow> OnGenerateRow(FOperatorStackEditorItemPtr InItem, const TSharedRef<STableViewBase>& InOwnerTable);
 
@@ -147,8 +150,14 @@ private:
 	void OnSearchTextCommitted(const FText& InFilterText, ETextCommit::Type InCommitType);
 	void OnSearchPinnedKeyword(ECheckBoxState InCheckState, FString InPinnedKeyword);
 
-	void OnSearchChanged();
+	/** Request an async search on the main panel */
+	void RequestSearchAsync() const;
+
+	/** Handles search in a recursive way, hiding items that do not match keywords */
+	bool HandleRecursiveSearch(const TSet<FString>& InSearchedKeywords_OR, const TSet<FString>& InSearchedKeywords_AND);
+
+	/** Returns true when this item matches the keywords */
 	bool MatchSearch(const TSet<FString>& InSearchedKeywordsOR, const TSet<FString>& InSearchedKeywordsAND) const;
-	bool HandleSearch(const TSet<FString>& InSearchedKeywordsOR, const TSet<FString>& InSearchedKeywordsAND);
+
 };
 
