MotionDesign : OperatorStack
- Fixed invalid underlying items being returned (object, struct, primitive)

diff --git a/Engine/Plugins/Experimental/OperatorStack/Source/OperatorStack/Private/Items/OperatorStackEditorTree.cpp b/Engine/Plugins/Experimental/OperatorStack/Source/OperatorStack/Private/Items/OperatorStackEditorTree.cpp
--- a/Engine/Plugins/Experimental/OperatorStack/Source/OperatorStack/Private/Items/OperatorStackEditorTree.cpp
+++ b/Engine/Plugins/Experimental/OperatorStack/Source/OperatorStack/Private/Items/OperatorStackEditorTree.cpp
@@ -21,7 +21,9 @@ FOperatorStackEditorTree::FOperatorStackEditorTree(UOperatorStackEditorStackCust
 
 	// Get root items
 	TSharedPtr<FOperatorStackEditorItem> RootItem;
-	if (InCustomization->GetRootItem(*InContext.Get(), RootItem))
+	if (InCustomization->GetRootItem(*InContext.Get(), RootItem)
+		&& RootItem.IsValid()
+		&& RootItem->HasValue())
 	{
 		check(InCustomization->IsCustomizationSupportedFor(RootItem))
 
diff --git a/Engine/Plugins/Experimental/OperatorStack/Source/OperatorStack/Private/Widgets/TableRows/SOperatorStackEditorStackRow.cpp b/Engine/Plugins/Experimental/OperatorStack/Source/OperatorStack/Private/Widgets/TableRows/SOperatorStackEditorStackRow.cpp
--- a/Engine/Plugins/Experimental/OperatorStack/Source/OperatorStack/Private/Widgets/TableRows/SOperatorStackEditorStackRow.cpp
+++ b/Engine/Plugins/Experimental/OperatorStack/Source/OperatorStack/Private/Widgets/TableRows/SOperatorStackEditorStackRow.cpp
@@ -27,10 +27,7 @@ void SOperatorStackEditorStackRow::Construct(const FArguments& InArgs
 	ChildSlot
 	[
 		SNew(SVerticalBox)
-		.Visibility(MakeAttributeLambda([this]()->EVisibility
-		{
-			return InnerStack->GetVisibility();
-		}))
+		.Visibility(InnerStack.Get(), &SOperatorStackEditorStack::GetVisibility)
 		+ SVerticalBox::Slot()
 		.AutoHeight()
 		[
diff --git a/Engine/Plugins/Experimental/OperatorStack/Source/OperatorStack/Public/Items/OperatorStackEditorGroupItem.h b/Engine/Plugins/Experimental/OperatorStack/Source/OperatorStack/Public/Items/OperatorStackEditorGroupItem.h
--- a/Engine/Plugins/Experimental/OperatorStack/Source/OperatorStack/Public/Items/OperatorStackEditorGroupItem.h
+++ b/Engine/Plugins/Experimental/OperatorStack/Source/OperatorStack/Public/Items/OperatorStackEditorGroupItem.h
@@ -35,17 +35,15 @@ struct FOperatorStackEditorGroupItem : FOperatorStackEditorItem
 		return Items.Num();
 	}
 
-	virtual bool HasValue() const override
+	virtual bool HasValue(uint32 InIndex) const override
 	{
-		for (const FOperatorStackEditorItemPtr& Item : Items)
+		if (Items.IsValidIndex(InIndex) && Items[InIndex].IsValid())
 		{
-			if (Item.IsValid() && Item->HasValue())
-			{
-				return true;
-			}
+			// No groups allowed so pick first item
+			return Items[InIndex]->HasValue(0);
 		}
 
-		return FOperatorStackEditorItem::HasValue();
+		return FOperatorStackEditorItem::HasValue(InIndex);
 	}
 
 	virtual uint32 GetHash() const override
@@ -53,7 +51,7 @@ struct FOperatorStackEditorGroupItem : FOperatorStackEditorItem
 		return CachedHash;
 	}
 
-	virtual void* GetValuePtr(int32 InIndex) const override
+	virtual void* GetValuePtr(uint32 InIndex) const override
 	{
 		if (Items.IsValidIndex(InIndex) && Items[InIndex].IsValid())
 		{
diff --git a/Engine/Plugins/Experimental/OperatorStack/Source/OperatorStack/Public/Items/OperatorStackEditorItem.h b/Engine/Plugins/Experimental/OperatorStack/Source/OperatorStack/Public/Items/OperatorStackEditorItem.h
--- a/Engine/Plugins/Experimental/OperatorStack/Source/OperatorStack/Public/Items/OperatorStackEditorItem.h
+++ b/Engine/Plugins/Experimental/OperatorStack/Source/OperatorStack/Public/Items/OperatorStackEditorItem.h
@@ -75,7 +75,7 @@ struct FOperatorStackEditorItem
 	{
 		check(InIndex < GetValueCount())
 
-		if (!GetValuePtr(InIndex))
+		if (!HasValue(InIndex))
 		{
 			return nullptr;
 		}
@@ -118,7 +118,10 @@ struct FOperatorStackEditorItem
 
 		for (uint32 Index = 0; Index < ValueCount; Index++)
 		{
-			Values.Add(Get<InValueType>(Index));
+			if (InValueType* Value = Get<InValueType>(Index))
+			{
+				Values.Add(Value);
+			}
 		}
 
 		return Values;
@@ -136,14 +139,27 @@ struct FOperatorStackEditorItem
 		return 0;
 	}
 
+	/** Checks if this item has at least one value and that it is usable */
+	bool HasValue() const
+	{
+		bool bHasValue = false;
+
+		for (uint32 Index = 0; Index < GetValueCount(); Index++)
+		{
+			bHasValue |= HasValue(Index);
+		}
+
+		return bHasValue;
+	}
+
 	/** Checks if this item has a value and that it is usable */
-	virtual bool HasValue() const
+	virtual bool HasValue(uint32 InIndex) const
 	{
 		return false;
 	}
 
 	/** Get raw ptr to value, prefer using Get<>() instead */
-	virtual void* GetValuePtr(int32 InIndex) const
+	virtual void* GetValuePtr(uint32 InIndex) const
 	{
 		return nullptr;
 	}
diff --git a/Engine/Plugins/Experimental/OperatorStack/Source/OperatorStack/Public/Items/OperatorStackEditorObjectItem.h b/Engine/Plugins/Experimental/OperatorStack/Source/OperatorStack/Public/Items/OperatorStackEditorObjectItem.h
--- a/Engine/Plugins/Experimental/OperatorStack/Source/OperatorStack/Public/Items/OperatorStackEditorObjectItem.h
+++ b/Engine/Plugins/Experimental/OperatorStack/Source/OperatorStack/Public/Items/OperatorStackEditorObjectItem.h
@@ -30,7 +30,7 @@ struct FOperatorStackEditorObjectItem : FOperatorStackEditorItem
 		return 1;
 	}
 
-	virtual bool HasValue() const override
+	virtual bool HasValue(uint32 InIndex) const override
 	{
 		return ObjectWeak.IsValid();
 	}
@@ -40,7 +40,7 @@ struct FOperatorStackEditorObjectItem : FOperatorStackEditorItem
 		return CachedHash;
 	}
 
-	virtual void* GetValuePtr(int32 InIndex) const override
+	virtual void* GetValuePtr(uint32 InIndex) const override
 	{
 		return ObjectWeak.Get();
 	}
diff --git a/Engine/Plugins/Experimental/OperatorStack/Source/OperatorStack/Public/Items/OperatorStackEditorPrimitiveItem.h b/Engine/Plugins/Experimental/OperatorStack/Source/OperatorStack/Public/Items/OperatorStackEditorPrimitiveItem.h
--- a/Engine/Plugins/Experimental/OperatorStack/Source/OperatorStack/Public/Items/OperatorStackEditorPrimitiveItem.h
+++ b/Engine/Plugins/Experimental/OperatorStack/Source/OperatorStack/Public/Items/OperatorStackEditorPrimitiveItem.h
@@ -30,7 +30,7 @@ struct FOperatorStackEditorPrimitiveItem : FOperatorStackEditorItem
 		return 1;
 	}
 
-	virtual bool HasValue() const override
+	virtual bool HasValue(uint32 InIndex) const override
 	{
 		return MemoryPtr != nullptr;
 	}
@@ -40,7 +40,7 @@ struct FOperatorStackEditorPrimitiveItem : FOperatorStackEditorItem
 		return CachedHash;
 	}
 
-	virtual void* GetValuePtr(int32 InIndex) const override
+	virtual void* GetValuePtr(uint32 InIndex) const override
 	{
 		return MemoryPtr;
 	}
diff --git a/Engine/Plugins/Experimental/OperatorStack/Source/OperatorStack/Public/Items/OperatorStackEditorStructItem.h b/Engine/Plugins/Experimental/OperatorStack/Source/OperatorStack/Public/Items/OperatorStackEditorStructItem.h
--- a/Engine/Plugins/Experimental/OperatorStack/Source/OperatorStack/Public/Items/OperatorStackEditorStructItem.h
+++ b/Engine/Plugins/Experimental/OperatorStack/Source/OperatorStack/Public/Items/OperatorStackEditorStructItem.h
@@ -70,7 +70,7 @@ struct FOperatorStackEditorStructItem : FOperatorStackEditorItem
 		return 1;
 	}
 
-	virtual bool HasValue() const override
+	virtual bool HasValue(uint32 InIndex) const override
 	{
 		if (const TSharedPtr<FStructOnScope> Struct = StructWeak.Pin())
 		{
@@ -85,7 +85,7 @@ struct FOperatorStackEditorStructItem : FOperatorStackEditorItem
 		return CachedHash;
 	}
 
-	virtual void* GetValuePtr(int32 InIndex) const override
+	virtual void* GetValuePtr(uint32 InIndex) const override
 	{
 		if (const TSharedPtr<FStructOnScope> Struct = StructWeak.Pin())
 		{
