MotionDesign : OperatorStack
- Fixed performance drop significantly when enabling pinned filters in operator stack

diff --git a/Engine/Plugins/Experimental/OperatorStack/Source/OperatorStack/Private/Widgets/SOperatorStackEditorStack.cpp b/Engine/Plugins/Experimental/OperatorStack/Source/OperatorStack/Private/Widgets/SOperatorStackEditorStack.cpp
--- a/Engine/Plugins/Experimental/OperatorStack/Source/OperatorStack/Private/Widgets/SOperatorStackEditorStack.cpp
+++ b/Engine/Plugins/Experimental/OperatorStack/Source/OperatorStack/Private/Widgets/SOperatorStackEditorStack.cpp
@@ -705,14 +705,19 @@ TSharedRef<SWidget> SOperatorStackEditorStack::GenerateMenuWidget(FName InMenuNa
 	return UToolMenus::Get()->GenerateWidget(InMenuName, ToolMenuContext);
 }
 
+EVisibility SOperatorStackEditorStack::GetHeaderVisibility() const
+{
+	return !bHiddenByFilter ? EVisibility::Visible : EVisibility::Collapsed;
+}
+
 EVisibility SOperatorStackEditorStack::GetBodyVisibility() const
 {
-	return bHeaderExpanded ? EVisibility::Visible : EVisibility::Collapsed;
+	return !bHiddenByFilter && bHeaderExpanded ? EVisibility::Visible : EVisibility::Collapsed;
 }
 
 EVisibility SOperatorStackEditorStack::GetFooterVisibility() const
 {
-	return bHeaderExpanded ? EVisibility::Visible : EVisibility::Collapsed;
+	return !bHiddenByFilter && bHeaderExpanded ? EVisibility::Visible : EVisibility::Collapsed;
 }
 
 EVisibility SOperatorStackEditorStack::GetMessageBoxVisibility() const
@@ -928,7 +933,7 @@ bool SOperatorStackEditorStack::HandleRecursiveSearch(const TSet<FString>& InSea
 	if (CustomizeItem.IsValid())
 	{
 		bMatchSearch |= MatchSearch(InSearchedKeywords_OR, InSearchedKeywords_AND);
-		SetVisibility(bMatchSearch ? EVisibility::Visible : EVisibility::Collapsed);
+		bHiddenByFilter = !bMatchSearch;
 	}
 
 	return bMatchSearch;
@@ -982,36 +987,39 @@ TSharedRef<SWidget> SOperatorStackEditorStack::GenerateStackWidget()
 	if (CustomizeItem.IsValid() && Items.IsEmpty())
 	{
 		return SNew(SHorizontalBox)
-		+ SHorizontalBox::Slot()
-		.Padding(0.f)
-		.AutoWidth()
-		[
-			SNew(SSeparator)
-			.Visibility(BorderColor == FLinearColor::Transparent ? EVisibility::Collapsed : EVisibility::Visible)
-			.ColorAndOpacity(BorderColor)
-			.SeparatorImage(FAppStyle::GetBrush("ThinLine.Horizontal"))
-			.Thickness(3.0f)
-			.Orientation(EOrientation::Orient_Vertical)
-		]
-		+ SHorizontalBox::Slot()
-		.Padding(0.f)
-		.FillWidth(1.f)
-		[
-			SNew(SBox)
-			.Padding(2.f)
+			.Visibility(this, &SOperatorStackEditorStack::GetHeaderVisibility)
+
+			+ SHorizontalBox::Slot()
+			.Padding(0.f)
+			.AutoWidth()
 			[
-				SNew(SOverlay)
-				+ SOverlay::Slot()
-				[
-					SNew(SColorBlock)
-					.Color(FOperatorStackEditorStyle::Get().GetColor("ForegroundColor"))
-				]
-				+ SOverlay::Slot()
+				SNew(SSeparator)
+				.Visibility(BorderColor == FLinearColor::Transparent ? EVisibility::Collapsed : EVisibility::Visible)
+				.ColorAndOpacity(BorderColor)
+				.SeparatorImage(FAppStyle::GetBrush("ThinLine.Horizontal"))
+				.Thickness(3.0f)
+				.Orientation(EOrientation::Orient_Vertical)
+			]
+
+			+ SHorizontalBox::Slot()
+			.Padding(0.f)
+			.FillWidth(1.f)
+			[
+				SNew(SBox)
+				.Padding(2.f)
 				[
-					StackWidget
+					SNew(SOverlay)
+					+ SOverlay::Slot()
+					[
+						SNew(SColorBlock)
+						.Color(FOperatorStackEditorStyle::Get().GetColor("ForegroundColor"))
+					]
+					+ SOverlay::Slot()
+					[
+						StackWidget
+					]
 				]
-			]
-		];
+			];
 	}
 
 	return StackWidget;
diff --git a/Engine/Plugins/Experimental/OperatorStack/Source/OperatorStack/Private/Widgets/SOperatorStackEditorStack.h b/Engine/Plugins/Experimental/OperatorStack/Source/OperatorStack/Private/Widgets/SOperatorStackEditorStack.h
--- a/Engine/Plugins/Experimental/OperatorStack/Source/OperatorStack/Private/Widgets/SOperatorStackEditorStack.h
+++ b/Engine/Plugins/Experimental/OperatorStack/Source/OperatorStack/Private/Widgets/SOperatorStackEditorStack.h
@@ -81,6 +81,7 @@ protected:
 	/** Generate this named menu widget for this context */
 	TSharedRef<SWidget> GenerateMenuWidget(FName InMenuName) const;
 
+	EVisibility GetHeaderVisibility() const;
 	EVisibility GetBodyVisibility() const;
 	EVisibility GetFooterVisibility() const;
 
@@ -143,6 +144,8 @@ protected:
 	/** When set will display an alert message in the header */
 	TAttribute<EOperatorStackEditorMessageType> MessageBoxType;
 
+	bool bHiddenByFilter = false;
+
 private:
 	static TSharedRef<ICustomDetailsView> CreateDetailsView(const FCustomDetailsViewArgs& InArgs, const FOperatorStackEditorItem& InItem);
 
@@ -158,6 +161,5 @@ private:
 
 	/** Returns true when this item matches the keywords */
 	bool MatchSearch(const TSet<FString>& InSearchedKeywordsOR, const TSet<FString>& InSearchedKeywordsAND) const;
-
 };
 
