MotionDesign : OperatorStack
- Fix crash on mac when deleting items from operator stack, crash is due to filtering details view property on an invalid object

diff --git a/Engine/Plugins/Experimental/Avalanche/Source/AvalancheEditor/Private/OperatorStack/Widgets/SAvaOperatorStackTab.cpp b/Engine/Plugins/Experimental/Avalanche/Source/AvalancheEditor/Private/OperatorStack/Widgets/SAvaOperatorStackTab.cpp
--- a/Engine/Plugins/Experimental/Avalanche/Source/AvalancheEditor/Private/OperatorStack/Widgets/SAvaOperatorStackTab.cpp
+++ b/Engine/Plugins/Experimental/Avalanche/Source/AvalancheEditor/Private/OperatorStack/Widgets/SAvaOperatorStackTab.cpp
@@ -128,7 +128,11 @@ void SAvaOperatorStackTab::OnAnimatorRemoved(UPropertyAnimatorCoreComponent* InC
 		if (SelectionSet->CountSelections(UPropertyAnimatorCoreComponent::StaticClass())
 			|| SelectionSet->CountSelections(UPropertyAnimatorCoreBase::StaticClass()))
 		{
-			RefreshSelection(SelectionSet);
+			FTSTicker::GetCoreTicker().AddTicker(FTickerDelegate::CreateSPLambda(this, [this, SelectionSet](float)
+			{
+				RefreshSelection(SelectionSet);
+				return false;
+			}));
 		}
 		else
 		{
diff --git a/Engine/Plugins/Experimental/OperatorStack/Source/OperatorStack/Private/Widgets/SOperatorStackEditorStack.cpp b/Engine/Plugins/Experimental/OperatorStack/Source/OperatorStack/Private/Widgets/SOperatorStackEditorStack.cpp
--- a/Engine/Plugins/Experimental/OperatorStack/Source/OperatorStack/Private/Widgets/SOperatorStackEditorStack.cpp
+++ b/Engine/Plugins/Experimental/OperatorStack/Source/OperatorStack/Private/Widgets/SOperatorStackEditorStack.cpp
@@ -911,27 +911,27 @@ bool SOperatorStackEditorStack::HandleRecursiveSearch(const TSet<FString>& InSea
 		}
 	}
 
-	// Search for properties row in details view
-	bool bRowFound = false;
-
-	if (BodyDetailsView)
+	// Do not hide root item if nothing was found
+	if (CustomizeItem.IsValid() && CustomizeItem->HasValue())
 	{
-		bRowFound |= BodyDetailsView->FilterItems(InSearchedKeywords_AND.Array());
-	}
+		// Search for properties row in details view
+		bool bRowFound = false;
 
-	if (FooterDetailsView)
-	{
-		bRowFound |= FooterDetailsView->FilterItems(InSearchedKeywords_AND.Array());
-	}
+		if (BodyDetailsView)
+		{
+			bRowFound |= BodyDetailsView->FilterItems(InSearchedKeywords_AND.Array());
+		}
 
-	if (!InSearchedKeywords_AND.IsEmpty())
-	{
-		bMatchSearch |= bRowFound;
-	}
+		if (FooterDetailsView)
+		{
+			bRowFound |= FooterDetailsView->FilterItems(InSearchedKeywords_AND.Array());
+		}
 
-	// Do not hide root item if nothing was found
-	if (CustomizeItem.IsValid())
-	{
+		if (!InSearchedKeywords_AND.IsEmpty())
+		{
+			bMatchSearch |= bRowFound;
+		}
+		
 		bMatchSearch |= MatchSearch(InSearchedKeywords_OR, InSearchedKeywords_AND);
 		bHiddenByFilter = !bMatchSearch;
 	}
