MotionDesign : OperatorStack
- OperatorStack can now be used outside of Motion Design levels, code decoupled for standalone usage
- Tab will be auto invoked when modifiers or animators are updated or on selection change

diff --git a/Engine/Plugins/Experimental/ActorModifierCore/Source/ActorModifierCoreEditor/Private/Modifiers/Customizations/ActorModifierCoreEditorStackCustomization.cpp b/Engine/Plugins/Experimental/ActorModifierCore/Source/ActorModifierCoreEditor/Private/Modifiers/Customizations/ActorModifierCoreEditorStackCustomization.cpp
--- a/Engine/Plugins/Experimental/ActorModifierCore/Source/ActorModifierCoreEditor/Private/Modifiers/Customizations/ActorModifierCoreEditorStackCustomization.cpp
+++ b/Engine/Plugins/Experimental/ActorModifierCore/Source/ActorModifierCoreEditor/Private/Modifiers/Customizations/ActorModifierCoreEditorStackCustomization.cpp
@@ -36,6 +36,20 @@ UActorModifierCoreEditorStackCustomization::UActorModifierCoreEditorStackCustomi
 {
 	// for stack and modifiers
 	RegisterCustomizationFor(UActorModifierCoreBase::StaticClass());
+
+	// Modifiers delegates
+	UActorModifierCoreStack::OnModifierAdded().AddUObject(this, &UActorModifierCoreEditorStackCustomization::OnModifierUpdated);
+	UActorModifierCoreStack::OnModifierMoved().AddUObject(this, &UActorModifierCoreEditorStackCustomization::OnModifierUpdated);
+	UActorModifierCoreStack::OnModifierRemoved().AddUObject(this, &UActorModifierCoreEditorStackCustomization::OnModifierRemoved);
+	UActorModifierCoreStack::OnModifierReplaced().AddUObject(this, &UActorModifierCoreEditorStackCustomization::OnModifierUpdated);
+}
+
+UActorModifierCoreEditorStackCustomization::~UActorModifierCoreEditorStackCustomization()
+{
+	UActorModifierCoreStack::OnModifierAdded().RemoveAll(this);
+	UActorModifierCoreStack::OnModifierMoved().RemoveAll(this);
+	UActorModifierCoreStack::OnModifierRemoved().RemoveAll(this);
+	UActorModifierCoreStack::OnModifierReplaced().RemoveAll(this);
 }
 
 bool UActorModifierCoreEditorStackCustomization::GetRootItem(const FOperatorStackEditorContext& InContext, FOperatorStackEditorItemPtr& OutRootItem) const
@@ -473,6 +487,17 @@ void UActorModifierCoreEditorStackCustomization::OnDropItem(const TArray<FOperat
 	Super::OnDropItem(InDraggedItems, InDropZoneItem, InZone);
 }
 
+bool UActorModifierCoreEditorStackCustomization::ShouldFocusCustomization(const FOperatorStackEditorContext& InContext) const
+{
+	if (InContext.GetItems().IsEmpty())
+	{
+		return false;
+	}
+	
+	return InContext.GetItems().Last().IsValid()
+		&& (InContext.GetItems().Last()->IsA<UActorModifierCoreBase>() || InContext.GetItems().Last()->IsA<UActorModifierCoreComponent>());
+}
+
 void UActorModifierCoreEditorStackCustomization::FillStackHeaderMenu(UToolMenu* InToolMenu) const
 {
 	if (!InToolMenu)
@@ -849,6 +874,22 @@ void UActorModifierCoreEditorStackCustomization::ToggleModifierProfilingAction(F
 	}
 }
 
+void UActorModifierCoreEditorStackCustomization::OnModifierUpdated(UActorModifierCoreBase* InModifier)
+{
+	if (InModifier)
+	{
+		RefreshActiveSelection(InModifier->GetRootModifierStack(), /** Force */false);
+	}
+}
+
+void UActorModifierCoreEditorStackCustomization::OnModifierRemoved(UActorModifierCoreBase* InModifier)
+{
+	if (InModifier)
+	{
+		RefreshActiveSelection(InModifier->GetRootModifierStack(), /** Force */true);
+	}
+}
+
 TSharedRef<FUICommandList> UActorModifierCoreEditorStackCustomization::CreateModifierCommands(FOperatorStackEditorItemPtr InItem)
 {
 	TSharedRef<FUICommandList> Commands = MakeShared<FUICommandList>();
diff --git a/Engine/Plugins/Experimental/ActorModifierCore/Source/ActorModifierCoreEditor/Private/Modifiers/Customizations/ActorModifierCoreEditorStackCustomization.h b/Engine/Plugins/Experimental/ActorModifierCore/Source/ActorModifierCoreEditor/Private/Modifiers/Customizations/ActorModifierCoreEditorStackCustomization.h
--- a/Engine/Plugins/Experimental/ActorModifierCore/Source/ActorModifierCoreEditor/Private/Modifiers/Customizations/ActorModifierCoreEditorStackCustomization.h
+++ b/Engine/Plugins/Experimental/ActorModifierCore/Source/ActorModifierCoreEditor/Private/Modifiers/Customizations/ActorModifierCoreEditorStackCustomization.h
@@ -34,6 +34,7 @@ public:
 	static inline const FString PropertiesWrapperPrefix = TEXT("ActorModifierCoreEditorPropertiesWrapper");
 
 	UActorModifierCoreEditorStackCustomization();
+	virtual ~UActorModifierCoreEditorStackCustomization() override;
 
 	//~ Begin UOperatorStackEditorStackCustomization
 	virtual bool GetRootItem(const FOperatorStackEditorContext& InContext, FOperatorStackEditorItemPtr& OutRootItem) const override;
@@ -45,6 +46,7 @@ public:
 	virtual bool OnIsItemDraggable(const FOperatorStackEditorItemPtr& InDragItem) override;
 	virtual TOptional<EItemDropZone> OnItemCanAcceptDrop(const TArray<FOperatorStackEditorItemPtr>& InDraggedItems, const FOperatorStackEditorItemPtr& InDropZoneItem, EItemDropZone InZone) override;
 	virtual void OnDropItem(const TArray<FOperatorStackEditorItemPtr>& InDraggedItems, const FOperatorStackEditorItemPtr& InDropZoneItem, EItemDropZone InZone) override;
+	virtual bool ShouldFocusCustomization(const FOperatorStackEditorContext& InContext) const override;
 	//~ UOperatorStackEditorStackCustomization
 
 protected:
@@ -73,6 +75,9 @@ protected:
 	bool IsModifierProfiling(FOperatorStackEditorItemPtr InItem) const;
 	void ToggleModifierProfilingAction(FOperatorStackEditorItemPtr InItem) const;
 
+	void OnModifierUpdated(UActorModifierCoreBase* InModifier);
+	void OnModifierRemoved(UActorModifierCoreBase* InModifier);
+
 	TSharedRef<FUICommandList> CreateModifierCommands(FOperatorStackEditorItemPtr InItem);
 
 	bool CreatePropertiesHandlesMapFromModifier(UActorModifierCoreBase* InModifier, TMap<FName, FString>& OutModifierPropertiesHandlesMap) const;
diff --git a/Engine/Plugins/Experimental/Avalanche/Source/AvalancheEditor/Private/OperatorStack/AvaOperatorStackExtension.cpp b/Engine/Plugins/Experimental/Avalanche/Source/AvalancheEditor/Private/OperatorStack/AvaOperatorStackExtension.cpp
--- a/Engine/Plugins/Experimental/Avalanche/Source/AvalancheEditor/Private/OperatorStack/AvaOperatorStackExtension.cpp
+++ b/Engine/Plugins/Experimental/Avalanche/Source/AvalancheEditor/Private/OperatorStack/AvaOperatorStackExtension.cpp
@@ -1,22 +1,19 @@
 // Copyright Epic Games, Inc. All Rights Reserved.
 
 #include "AvaOperatorStackExtension.h"
-#include "AvaEditorCommands.h"
+
 #include "AvaLevelViewportCommands.h"
-#include "AvaOperatorStackTabSpawner.h"
-#include "Components/PropertyAnimatorCoreComponent.h"
+#include "DetailView/AvaDetailsExtension.h"
 #include "EditorModeManager.h"
 #include "Framework/Docking/LayoutExtender.h"
 #include "LevelEditor.h"
-#include "Modifiers/ActorModifierCoreBase.h"
-#include "Modifiers/ActorModifierCoreComponent.h"
 #include "Selection.h"
 #include "Subsystems/OperatorStackEditorSubsystem.h"
 #include "Subsystems/PropertyAnimatorCoreSubsystem.h"
 #include "ToolMenu.h"
 #include "ToolMenuEntry.h"
 #include "ToolMenuSection.h"
-#include "Widgets/SAvaOperatorStackTab.h"
+#include "Widgets/Docking/SDockTab.h"
 #include "Widgets/SOperatorStackEditorWidget.h"
 
 #define LOCTEXT_NAMESPACE "AvaOperatorStackExtension"
@@ -26,49 +23,68 @@ FAvaOperatorStackExtension::FAvaOperatorStackExtension()
 {
 }
 
-TSharedPtr<SDockTab> FAvaOperatorStackExtension::OpenTab()
+TSharedPtr<SDockTab> FAvaOperatorStackExtension::FindOrOpenTab(bool bInOpen) const
 {
 	const TSharedPtr<IAvaEditor>& Editor = GetEditor();
-
 	if (!Editor.IsValid())
 	{
 		return nullptr;
 	}
 
-	TSharedPtr<FTabManager> TabManager = Editor->GetTabManager();
-
+	const TSharedPtr<FTabManager> TabManager = Editor->GetTabManager();
 	if (!TabManager.IsValid())
 	{
 		return nullptr;
 	}
 
-	return TabManager->TryInvokeTab(FAvaOperatorStackTabSpawner::GetTabID());
+	TSharedPtr<SDockTab> Tab = TabManager->FindExistingLiveTab(UOperatorStackEditorSubsystem::TabId);
+
+	if (!Tab.IsValid() && bInOpen)
+	{
+		Tab = TabManager->TryInvokeTab(UOperatorStackEditorSubsystem::TabId);
+	}
+
+	if (!Tab.IsValid())
+	{
+		return nullptr;
+	}
+
+	const TSharedRef<SOperatorStackEditorWidget> Widget = StaticCastSharedRef<SOperatorStackEditorWidget>(Tab->GetContent());
+
+	if (const TSharedPtr<FAvaDetailsExtension> DetailsExtension = Editor->FindExtension<FAvaDetailsExtension>())
+	{
+		Widget->SetKeyframeHandler(DetailsExtension->GetDetailsKeyframeHandler());
+	}
+
+	return Tab;
 }
 
 void FAvaOperatorStackExtension::Activate()
 {
 	FAvaEditorExtension::Activate();
 
-	USelection::SelectionChangedEvent.AddSP(this, &FAvaOperatorStackExtension::OnSelectionChanged);
+	if (!FindOrOpenTab(/** Open */false))
+	{
+		UOperatorStackEditorSubsystem::OnOperatorStackSpawned().AddSP(this, &FAvaOperatorStackExtension::OnOperatorStackSpawned);
+	}
 }
 
 void FAvaOperatorStackExtension::Deactivate()
 {
 	FAvaEditorExtension::Deactivate();
 
-	USelection::SelectionChangedEvent.RemoveAll(this);
+	UOperatorStackEditorSubsystem::OnOperatorStackSpawned().RemoveAll(this);
 }
 
 void FAvaOperatorStackExtension::RegisterTabSpawners(const TSharedRef<IAvaEditor>& InEditor) const
 {
-	InEditor->AddTabSpawner<FAvaOperatorStackTabSpawner>(InEditor);
 }
 
 void FAvaOperatorStackExtension::ExtendLevelEditorLayout(FLayoutExtender& InExtender) const
 {
 	InExtender.ExtendLayout(LevelEditorTabIds::LevelEditorSceneOutliner
 		, ELayoutExtensionPosition::After
-		, FTabManager::FTab(FAvaOperatorStackTabSpawner::GetTabID(), ETabState::ClosedTab));
+		, FTabManager::FTab(UOperatorStackEditorSubsystem::TabId, ETabState::ClosedTab));
 }
 
 void FAvaOperatorStackExtension::ExtendToolbarMenu(UToolMenu& InMenu)
@@ -77,7 +93,7 @@ void FAvaOperatorStackExtension::ExtendToolbarMenu(UToolMenu& InMenu)
 
 	FToolMenuEntry& Entry = Section.AddEntry(FToolMenuEntry::InitToolBarButton(
 		TEXT("OpenOperatorStackButton"),
-		FExecuteAction::CreateSPLambda(this, [this]() { OpenTab(); }),
+		FExecuteAction::CreateSPLambda(this, [this]() { FindOrOpenTab(/** Open */true); }),
 		LOCTEXT("OpenOperatorStackLabel", "Operator Stack"),
 		LOCTEXT("OpenOperatorStackTooltip", "Open the operator stack tab."),
 		FSlateIcon(FAppStyle::GetAppStyleSetName(), "ClassIcon.UserDefinedStruct")
@@ -99,69 +115,6 @@ void FAvaOperatorStackExtension::BindCommands(const TSharedRef<FUICommandList>&
 		, FExecuteAction::CreateSP(this, &FAvaOperatorStackExtension::EnableAnimators, true));
 }
 
-void FAvaOperatorStackExtension::OnSelectionChanged(UObject* InSelection)
-{
-	const USelection* Selection = Cast<USelection>(InSelection);
-
-	if (!Selection || Selection->Num() == 0)
-	{
-		return;
-	}
-
-	const UObject* LastSelectedObject = Selection->GetSelectedObject(Selection->Num() - 1);
-
-	if (!LastSelectedObject)
-	{
-		return;
-	}
-
-	const UOperatorStackEditorSubsystem* OperatorStackEditorSubsystem = UOperatorStackEditorSubsystem::Get();
-
-	if (!OperatorStackEditorSubsystem)
-	{
-		return;
-	}
-
-	FName ActiveCustomization;
-	if (LastSelectedObject->IsA<UActorModifierCoreBase>()
-		|| LastSelectedObject->IsA<UActorModifierCoreComponent>())
-	{
-		ActiveCustomization = TEXT("Modifiers");
-	}
-	else if (LastSelectedObject->IsA<UPropertyAnimatorCoreBase>()
-		|| LastSelectedObject->IsA<UPropertyAnimatorCoreComponent>())
-	{
-		ActiveCustomization = TEXT("Animators");
-	}
-
-	if (ActiveCustomization.IsNone())
-	{
-		return;
-	}
-
-	// Lets invoke the operator stack tab
-	const TSharedPtr<SDockTab> InvokedTab = OpenTab();
-
-	if (!InvokedTab.IsValid())
-	{
-		return;
-	}
-
-	// Find the avalanche operator stack and set the active customization
-	OperatorStackEditorSubsystem->ForEachCustomizationWidget([ActiveCustomization](TSharedRef<SOperatorStackEditorWidget> InWidget)->bool
-	{
-		if (InWidget->GetPanelTag() == SAvaOperatorStackTab::PanelTag)
-		{
-			InWidget->SetActiveCustomization(ActiveCustomization);
-
-			// Stop looping for customization widgets
-			return false;
-		}
-
-		return true;
-	});
-}
-
 void FAvaOperatorStackExtension::EnableAnimators(bool bInEnable) const
 {
 	const UWorld* const World = GetWorld();
@@ -202,4 +155,10 @@ void FAvaOperatorStackExtension::EnableAnimators(bool bInEnable) const
 	}
 }
 
+void FAvaOperatorStackExtension::OnOperatorStackSpawned(TSharedRef<SOperatorStackEditorWidget> InWidget)
+{
+	// Init tab with MD
+	FindOrOpenTab(/** Open */false);
+}
+
 #undef LOCTEXT_NAMESPACE
diff --git a/Engine/Plugins/Experimental/Avalanche/Source/AvalancheEditor/Private/OperatorStack/AvaOperatorStackExtension.h b/Engine/Plugins/Experimental/Avalanche/Source/AvalancheEditor/Private/OperatorStack/AvaOperatorStackExtension.h
--- a/Engine/Plugins/Experimental/Avalanche/Source/AvalancheEditor/Private/OperatorStack/AvaOperatorStackExtension.h
+++ b/Engine/Plugins/Experimental/Avalanche/Source/AvalancheEditor/Private/OperatorStack/AvaOperatorStackExtension.h
@@ -5,6 +5,7 @@
 #include "IAvaEditorExtension.h"
 
 class SDockTab;
+class SOperatorStackEditorWidget;
 
 class FAvaOperatorStackExtension : public FAvaEditorExtension
 {
@@ -13,7 +14,7 @@ public:
 
 	FAvaOperatorStackExtension();
 
-	TSharedPtr<SDockTab> OpenTab();
+	TSharedPtr<SDockTab> FindOrOpenTab(bool bInOpen) const;
 
 	//~ Begin IAvaEditorExtension
 	virtual void Activate() override;
@@ -25,8 +26,8 @@ public:
 	//~ End IAvaEditorExtension
 
 private:
-	void OnSelectionChanged(UObject* InSelection);
 	void EnableAnimators(bool bInEnable) const;
+	void OnOperatorStackSpawned(TSharedRef<SOperatorStackEditorWidget> InWidget);
 
 	TSharedRef<FUICommandList> AnimatorCommands;
 };
diff --git a/Engine/Plugins/Experimental/Avalanche/Source/AvalancheEditor/Private/OperatorStack/AvaOperatorStackTabSpawner.cpp b/Engine/Plugins/Experimental/Avalanche/Source/AvalancheEditor/Private/OperatorStack/AvaOperatorStackTabSpawner.cpp
deleted file mode 100644
diff --git a/Engine/Plugins/Experimental/Avalanche/Source/AvalancheEditor/Private/OperatorStack/AvaOperatorStackTabSpawner.h b/Engine/Plugins/Experimental/Avalanche/Source/AvalancheEditor/Private/OperatorStack/AvaOperatorStackTabSpawner.h
deleted file mode 100644
diff --git a/Engine/Plugins/Experimental/Avalanche/Source/AvalancheEditor/Private/OperatorStack/Widgets/SAvaOperatorStackTab.cpp b/Engine/Plugins/Experimental/Avalanche/Source/AvalancheEditor/Private/OperatorStack/Widgets/SAvaOperatorStackTab.cpp
deleted file mode 100644
diff --git a/Engine/Plugins/Experimental/Avalanche/Source/AvalancheEditor/Private/OperatorStack/Widgets/SAvaOperatorStackTab.h b/Engine/Plugins/Experimental/Avalanche/Source/AvalancheEditor/Private/OperatorStack/Widgets/SAvaOperatorStackTab.h
deleted file mode 100644
diff --git a/Engine/Plugins/Experimental/OperatorStack/Source/OperatorStack/OperatorStackEditor.Build.cs b/Engine/Plugins/Experimental/OperatorStack/Source/OperatorStack/OperatorStackEditor.Build.cs
--- a/Engine/Plugins/Experimental/OperatorStack/Source/OperatorStack/OperatorStackEditor.Build.cs
+++ b/Engine/Plugins/Experimental/OperatorStack/Source/OperatorStack/OperatorStackEditor.Build.cs
@@ -25,12 +25,14 @@ public class OperatorStackEditor : ModuleRules
 				"EditorSubsystem",
 				"Engine",
 				"InputCore",
+				"LevelEditor",
 				"Projects",
 				"PropertyEditor",
 				"Slate",
 				"SlateCore",
 				"ToolMenus",
 				"ToolWidgets",
+				"TypedElementRuntime",
 				"UnrealEd",
 				"WorkspaceMenuStructure"
 			}
diff --git a/Engine/Plugins/Experimental/OperatorStack/Source/OperatorStack/Private/Customizations/OperatorStackEditorStackCustomization.cpp b/Engine/Plugins/Experimental/OperatorStack/Source/OperatorStack/Private/Customizations/OperatorStackEditorStackCustomization.cpp
--- a/Engine/Plugins/Experimental/OperatorStack/Source/OperatorStack/Private/Customizations/OperatorStackEditorStackCustomization.cpp
+++ b/Engine/Plugins/Experimental/OperatorStack/Source/OperatorStack/Private/Customizations/OperatorStackEditorStackCustomization.cpp
@@ -2,6 +2,7 @@
 
 #include "Customizations/OperatorStackEditorStackCustomization.h"
 #include "Styling/SlateIconFinder.h"
+#include "Subsystems/OperatorStackEditorSubsystem.h"
 
 bool UOperatorStackEditorStackCustomization::RegisterCustomizationFor(const UStruct* InItemDefinition)
 {
@@ -73,3 +74,37 @@ const FSlateBrush* UOperatorStackEditorStackCustomization::GetIcon() const
 {
 	return FSlateIconFinder::FindIconBrushForClass(GetClass());
 }
+
+void UOperatorStackEditorStackCustomization::RefreshActiveSelection(UObject* InContext, bool bInForce) const
+{
+	if (!IsValid(InContext))
+	{
+		return;
+	}
+
+	UOperatorStackEditorSubsystem* Subsystem = UOperatorStackEditorSubsystem::Get();
+
+	if (!Subsystem)
+	{
+		return;
+	}
+
+	Subsystem->RefreshCustomizationWidget(InContext, bInForce);
+	Subsystem->FocusCustomizationWidget(InContext, Identifier);
+}
+
+void UOperatorStackEditorStackCustomization::FocusCustomization(const UObject* InContext) const
+{
+	if (!IsValid(InContext))
+	{
+		return;
+	}
+
+	UOperatorStackEditorSubsystem* Subsystem = UOperatorStackEditorSubsystem::Get();
+	if (!Subsystem)
+	{
+		return;
+	}
+
+	Subsystem->FocusCustomizationWidget(InContext, Identifier);
+}
diff --git a/Engine/Plugins/Experimental/OperatorStack/Source/OperatorStack/Private/Subsystems/OperatorStackEditorSubsystem.cpp b/Engine/Plugins/Experimental/OperatorStack/Source/OperatorStack/Private/Subsystems/OperatorStackEditorSubsystem.cpp
--- a/Engine/Plugins/Experimental/OperatorStack/Source/OperatorStack/Private/Subsystems/OperatorStackEditorSubsystem.cpp
+++ b/Engine/Plugins/Experimental/OperatorStack/Source/OperatorStack/Private/Subsystems/OperatorStackEditorSubsystem.cpp
@@ -2,14 +2,20 @@
 
 #include "Subsystems/OperatorStackEditorSubsystem.h"
 
-#include "Editor.h"
 #include "Customizations/OperatorStackEditorStackCustomization.h"
+#include "Editor.h"
+#include "LevelEditor.h"
+#include "Modules/ModuleManager.h"
 #include "UObject/UObjectIterator.h"
 #include "Widgets/SOperatorStackEditorPanel.h"
 #include "Widgets/SOperatorStackEditorWidget.h"
+#include "Widgets/Tabs/OperatorStackEditorTabInstance.h"
 
 DEFINE_LOG_CATEGORY_STATIC(LogOperatorStackEditorSubsystem, Log, All);
 
+const FName UOperatorStackEditorSubsystem::TabId = TEXT("OperatorStackEditorTab");
+UOperatorStackEditorSubsystem::FOnOperatorStackSpawned UOperatorStackEditorSubsystem::OnOperatorStackSpawnedDelegate;
+
 UOperatorStackEditorSubsystem* UOperatorStackEditorSubsystem::Get()
 {
 	if (GEditor)
@@ -24,13 +30,23 @@ void UOperatorStackEditorSubsystem::Initialize(FSubsystemCollectionBase& Collect
 	Super::Initialize(Collection);
 
 	ScanForStackCustomizations();
+
+	FLevelEditorModule& LevelEditorModule = FModuleManager::LoadModuleChecked<FLevelEditorModule>("LevelEditor");
+	LevelEditorModule.OnLevelEditorCreated().AddUObject(this, &UOperatorStackEditorSubsystem::OnLevelEditorCreated);
 }
 
 void UOperatorStackEditorSubsystem::Deinitialize()
 {
-	CustomizationStacks.Empty();
-
 	Super::Deinitialize();
+
+	CustomizationStacks.Empty();
+	CustomizationWidgets.Empty();
+	TabInstances.Empty();
+	
+	if (FLevelEditorModule* LevelEditorModule = FModuleManager::GetModulePtr<FLevelEditorModule>("LevelEditor"))
+	{
+		LevelEditorModule->OnLevelEditorCreated().RemoveAll(this);
+	}
 }
 
 bool UOperatorStackEditorSubsystem::RegisterStackCustomization(const TSubclassOf<UOperatorStackEditorStackCustomization> InStackCustomizationClass)
@@ -110,6 +126,36 @@ TSharedPtr<SOperatorStackEditorWidget> UOperatorStackEditorSubsystem::FindWidget
 	return nullptr;
 }
 
+TArray<TSharedPtr<SOperatorStackEditorWidget>> UOperatorStackEditorSubsystem::FindWidgets(const UWorld* InContext)
+{
+	TArray<TSharedPtr<SOperatorStackEditorWidget>> Widgets;
+	Widgets.Reserve(TabInstances.Num());
+
+	if (IsValid(InContext))
+	{
+		for (TSharedPtr<FOperatorStackEditorTabInstance>& TabInstance : TabInstances)
+		{
+			if (!TabInstance.IsValid())
+			{
+				continue;
+			}
+
+			TSharedPtr<ILevelEditor> LevelEditor = TabInstance->GetLevelEditor();
+			if (!LevelEditor.IsValid())
+			{
+				continue;
+			}
+
+			if (LevelEditor->GetWorld() == InContext)
+			{
+				Widgets.Emplace(TabInstance->GetOperatorStackEditorWidget());
+			}
+		}
+	}
+
+	return Widgets;
+}
+
 bool UOperatorStackEditorSubsystem::ForEachCustomization(TFunctionRef<bool(UOperatorStackEditorStackCustomization*)> InFunction) const
 {
 	TArray<UOperatorStackEditorStackCustomization*> Customizations;
@@ -159,6 +205,34 @@ bool UOperatorStackEditorSubsystem::ForEachCustomizationWidget(TFunctionRef<bool
 	return true;
 }
 
+void UOperatorStackEditorSubsystem::RefreshCustomizationWidget(UObject* InContext, bool bInForce)
+{
+	if (IsValid(InContext))
+	{
+		for (TSharedPtr<FOperatorStackEditorTabInstance>& TabInstance : TabInstances)
+		{
+			if (TabInstance.IsValid() && TabInstance->RefreshTab(InContext, bInForce))
+			{
+				return;
+			}
+		}
+	}
+}
+
+void UOperatorStackEditorSubsystem::FocusCustomizationWidget(const UObject* InContext, FName InIdentifier)
+{
+	if (IsValid(InContext))
+	{
+		for (TSharedPtr<FOperatorStackEditorTabInstance>& TabInstance : TabInstances)
+		{
+			if (TabInstance.IsValid() && TabInstance->FocusTab(InContext, InIdentifier))
+			{
+				return;
+			}
+		}
+	}
+}
+
 UOperatorStackEditorStackCustomization* UOperatorStackEditorSubsystem::GetCustomization(const FName& InName) const
 {
 	if (const TObjectPtr<UOperatorStackEditorStackCustomization>* Customization = CustomizationStacks.Find(InName))
@@ -185,3 +259,21 @@ void UOperatorStackEditorSubsystem::OnWidgetDestroyed(int32 InPanelId)
 {
 	CustomizationWidgets.Remove(InPanelId);
 }
+
+void UOperatorStackEditorSubsystem::OnLevelEditorCreated(TSharedPtr<ILevelEditor> InLevelEditor)
+{
+	for (TArray<TSharedPtr<FOperatorStackEditorTabInstance>>::TIterator It(TabInstances); It; ++It)
+	{
+		if (!It->IsValid() || !(*It)->GetLevelEditor().IsValid())
+		{
+			It.RemoveCurrent();
+		}
+	}
+
+	if (InLevelEditor.IsValid())
+	{
+		const TSharedRef<FOperatorStackEditorTabInstance> TabInstance = MakeShared<FOperatorStackEditorTabInstance>(InLevelEditor.ToSharedRef());
+		TabInstances.Emplace(TabInstance);
+		TabInstance->RegisterTab();
+	}
+}
diff --git a/Engine/Plugins/Experimental/OperatorStack/Source/OperatorStack/Private/Widgets/SOperatorStackEditorPanel.cpp b/Engine/Plugins/Experimental/OperatorStack/Source/OperatorStack/Private/Widgets/SOperatorStackEditorPanel.cpp
--- a/Engine/Plugins/Experimental/OperatorStack/Source/OperatorStack/Private/Widgets/SOperatorStackEditorPanel.cpp
+++ b/Engine/Plugins/Experimental/OperatorStack/Source/OperatorStack/Private/Widgets/SOperatorStackEditorPanel.cpp
@@ -92,7 +92,7 @@ void SOperatorStackEditorPanel::SetActiveCustomization(const FName& InCustomizat
 	}
 }
 
-void SOperatorStackEditorPanel::SetToolbarCustomization(const TArray<FName>& InCustomizations)
+void SOperatorStackEditorPanel::SetToolbarCustomizations(const TArray<FName>& InCustomizations)
 {
 	if (!WidgetToolbar.IsValid() || !WidgetSwitcher.IsValid())
 	{
@@ -125,6 +125,15 @@ void SOperatorStackEditorPanel::SetToolbarCustomization(const TArray<FName>& InC
 	UpdateSlots();
 }
 
+void SOperatorStackEditorPanel::GetToolbarCustomizations(TArray<UOperatorStackEditorStackCustomization*>& OutCustomizations)
+{
+	OutCustomizations.Empty(CustomizationStacks.Num());
+	Algo::Transform(CustomizationStacks, OutCustomizations, [](const TSharedRef<SOperatorStackEditorStack>& InStack)
+	{
+		return InStack->GetStackCustomization();
+	});
+}
+
 void SOperatorStackEditorPanel::ShowToolbarCustomization(FName InCustomization)
 {
 	if (const int32* CustomizationIdx = NamedStackIndexes.Find(InCustomization))
@@ -281,7 +290,7 @@ void SOperatorStackEditorPanel::UpdateSlots()
 
 			check(Customization)
 
-			CustomizationTrees.Add(Customization, FOperatorStackEditorTree(Customization, Context));
+			FOperatorStackEditorTree& ItemTree = CustomizationTrees.Add(Customization, FOperatorStackEditorTree(Customization, Context));
 
 			TSharedRef<SOperatorStackEditorStack> Stack = SNew(SOperatorStackEditorStack, SharedThis(this), Customization, nullptr);
 			CustomizationStacks.Add(Stack);
@@ -289,6 +298,11 @@ void SOperatorStackEditorPanel::UpdateSlots()
 			Box->SetContent(
 				Stack
 			);
+
+			if (Context.IsValid() && Customization->ShouldFocusCustomization(*Context))
+			{
+				SetActiveCustomization(Customization->GetIdentifier());
+			}
 		}
 	}
 }
diff --git a/Engine/Plugins/Experimental/OperatorStack/Source/OperatorStack/Private/Widgets/SOperatorStackEditorPanel.h b/Engine/Plugins/Experimental/OperatorStack/Source/OperatorStack/Private/Widgets/SOperatorStackEditorPanel.h
--- a/Engine/Plugins/Experimental/OperatorStack/Source/OperatorStack/Private/Widgets/SOperatorStackEditorPanel.h
+++ b/Engine/Plugins/Experimental/OperatorStack/Source/OperatorStack/Private/Widgets/SOperatorStackEditorPanel.h
@@ -28,46 +28,37 @@ public:
 
 	void Construct(const FArguments& InArgs);
 
+	//~ Begin SOperatorStackEditorWidget
 	virtual void SetContext(const FOperatorStackEditorContext& InContext) override;
-
 	virtual void SetActiveCustomization(const FName& InCustomization) override;
-
-	virtual void SetToolbarCustomization(const TArray<FName>& InCustomizations) override;
-
+	virtual void SetToolbarCustomizations(const TArray<FName>& InCustomizations) override;
+	virtual void GetToolbarCustomizations(TArray<UOperatorStackEditorStackCustomization*>& OutCustomizations) override;
 	virtual void ShowToolbarCustomization(FName InCustomization) override;
-
 	virtual void HideToolbarCustomization(FName InCustomization) override;
-
 	virtual void SetToolbarVisibility(bool bInVisible) override;
-
 	virtual void SetKeyframeHandler(TSharedPtr<IDetailKeyframeHandler> InKeyframeHandler) override;
-
 	virtual void SetDetailColumnSize(TSharedPtr<FDetailColumnSizeData> InDetailColumnSize) override;
-
 	virtual void SetPanelTag(FName InTag) override;
-
 	virtual void RefreshContext() override;
-
-	const FOperatorStackEditorTree& GetItemTree(UOperatorStackEditorStackCustomization* InCustomization);
-
-	void SaveItemExpansionState(uint32 InItem, bool bInExpanded);
-
-	bool GetItemExpansionState(uint32 InItem, bool& bOutExpanded);
-
 	virtual FOperatorStackEditorContextPtr GetContext() const override
 	{
 		return Context;
 	}
-
 	virtual int32 GetPanelId() const override
 	{
 		return PanelId;
 	}
-
 	virtual FName GetPanelTag() const override
 	{
 		return PanelTag;
 	}
+	//~ End SOperatorStackEditorWidget
+
+	const FOperatorStackEditorTree& GetItemTree(UOperatorStackEditorStackCustomization* InCustomization);
+
+	void SaveItemExpansionState(uint32 InItem, bool bInExpanded);
+
+	bool GetItemExpansionState(uint32 InItem, bool& bOutExpanded);
 
 	TSharedPtr<IDetailKeyframeHandler> GetKeyframeHandler() const
 	{
diff --git a/Engine/Plugins/Experimental/OperatorStack/Source/OperatorStack/Private/Widgets/Tabs/OperatorStackEditorTabInstance.cpp b/Engine/Plugins/Experimental/OperatorStack/Source/OperatorStack/Private/Widgets/Tabs/OperatorStackEditorTabInstance.cpp
new file mode 100644
--- /dev/null
+++ b/Engine/Plugins/Experimental/OperatorStack/Source/OperatorStack/Private/Widgets/Tabs/OperatorStackEditorTabInstance.cpp
@@ -0,0 +1,353 @@
+// Copyright Epic Games, Inc. All Rights Reserved.
+
+#include "Widgets/Tabs/OperatorStackEditorTabInstance.h"
+
+#include "Contexts/OperatorStackEditorContext.h"
+#include "Customizations/OperatorStackEditorStackCustomization.h"
+#include "Elements/Framework/TypedElementSelectionSet.h"
+#include "GameFramework/Actor.h"
+#include "ILevelEditor.h"
+#include "Items/OperatorStackEditorItem.h"
+#include "Items/OperatorStackEditorObjectItem.h"
+#include "Selection.h"
+#include "Subsystems/OperatorStackEditorSubsystem.h"
+#include "Widgets/Docking/SDockTab.h"
+#include "Widgets/SOperatorStackEditorWidget.h"
+#include "WorkspaceMenuStructure.h"
+#include "WorkspaceMenuStructureModule.h"
+
+#define LOCTEXT_NAMESPACE "OperatorStackEditorTabInstance"
+
+FOperatorStackEditorTabInstance::FOperatorStackEditorTabInstance(const TSharedRef<ILevelEditor>& InLevelEditor)
+	: LevelEditorWeak(InLevelEditor)
+{
+}
+
+FOperatorStackEditorTabInstance::~FOperatorStackEditorTabInstance()
+{
+	UnbindDelegates();
+	UnregisterTab();
+}
+
+TSharedPtr<SDockTab> FOperatorStackEditorTabInstance::InvokeTab()
+{
+	const TSharedPtr<ILevelEditor> LevelEditor = LevelEditorWeak.Pin();
+
+	if (!LevelEditor.IsValid())
+	{
+		return nullptr;
+	}
+
+	const TSharedPtr<FTabManager> TabManager = LevelEditor->GetTabManager();
+
+	if (!TabManager.IsValid())
+	{
+		return nullptr;
+	}
+
+	TSharedPtr<SDockTab> Tab = TabManager->TryInvokeTab(UOperatorStackEditorSubsystem::TabId);
+
+	if (Tab.IsValid())
+	{
+		Tab->ActivateInParent(ETabActivationCause::SetDirectly);
+		Tab->DrawAttention();
+	}
+
+	return Tab;
+}
+
+bool FOperatorStackEditorTabInstance::CloseTab()
+{
+	const TSharedPtr<ILevelEditor> LevelEditor = LevelEditorWeak.Pin();
+
+	if (!LevelEditor.IsValid())
+	{
+		return false;
+	}
+
+	const TSharedPtr<FTabManager> TabManager = LevelEditor->GetTabManager();
+
+	if (!TabManager.IsValid())
+	{
+		return false;
+	}
+
+	const TSharedPtr<SDockTab> Tab = TabManager->FindExistingLiveTab(UOperatorStackEditorSubsystem::TabId);
+
+	if (Tab.IsValid())
+	{
+		return Tab->RequestCloseTab();
+	}
+
+	return false;
+}
+
+bool FOperatorStackEditorTabInstance::RefreshTab(UObject* InContext, bool bInForce)
+{
+	const TSharedPtr<ILevelEditor> LevelEditor = LevelEditorWeak.Pin();
+	if (!LevelEditor.IsValid())
+	{
+		return false;
+	}
+
+	const UTypedElementSelectionSet* SelectionSet = LevelEditor->GetElementSelectionSet();
+	if (!SelectionSet)
+	{
+		return false;
+	}
+
+	const UWorld* ContextWorld = InContext->GetWorld();
+	if (LevelEditor->GetWorld() != ContextWorld)
+	{
+		return false;
+	}
+
+	if (bInForce)
+	{
+		OnSelectionSetChanged(SelectionSet);
+	}
+	else
+	{
+		const TSharedPtr<SOperatorStackEditorWidget> Widget = GetOperatorStackEditorWidget();
+		if (!Widget.IsValid())
+		{
+			return false;
+		}
+
+		const FOperatorStackEditorContextPtr Context = Widget->GetContext();
+		if (!Context.IsValid())
+		{
+			return false;
+		}
+
+		const AActor* ContextActor = InContext->GetTypedOuter<AActor>();
+		if (!ContextActor)
+		{
+			return false;
+		}
+
+		const TArray<UObject*> SelectedObjects = SelectionSet->GetSelectedObjects();
+		if (!SelectedObjects.Contains(ContextActor) && !SelectedObjects.Contains(InContext))
+		{
+			return false;
+		}
+
+		Widget->RefreshContext();
+	}
+	
+	return true;
+}
+
+bool FOperatorStackEditorTabInstance::FocusTab(const UObject* InContext, FName InIdentifier)
+{
+	const TSharedPtr<ILevelEditor> LevelEditor = LevelEditorWeak.Pin();
+	if (!LevelEditor.IsValid())
+	{
+		return false;
+	}
+
+	if (LevelEditor->GetWorld() != InContext->GetWorld())
+	{
+		return false;
+	}
+
+	InvokeTab();
+
+	if (!InIdentifier.IsNone())
+	{
+		if (const TSharedPtr<SOperatorStackEditorWidget> Widget = GetOperatorStackEditorWidget())
+		{
+			Widget->SetActiveCustomization(InIdentifier);
+		}
+	}
+
+	return true;
+}
+
+TSharedPtr<SOperatorStackEditorWidget> FOperatorStackEditorTabInstance::GetOperatorStackEditorWidget() const
+{
+	TSharedPtr<SOperatorStackEditorWidget> Widget;
+
+	if (WidgetIdentifier == INDEX_NONE)
+	{
+		return Widget;
+	}
+
+	if (UOperatorStackEditorSubsystem* Subsystem = UOperatorStackEditorSubsystem::Get())
+	{
+		Widget = Subsystem->FindWidget(WidgetIdentifier);
+	}
+
+	return Widget;
+}
+
+bool FOperatorStackEditorTabInstance::RegisterTab()
+{
+	const TSharedPtr<ILevelEditor> LevelEditor = LevelEditorWeak.Pin();
+
+	if (!LevelEditor.IsValid())
+	{
+		return false;
+	}
+
+	const TSharedPtr<FTabManager> TabManager = LevelEditor->GetTabManager();
+
+	if (!TabManager.IsValid())
+	{
+		return false;
+	}
+
+	static const FText DisplayNameText = LOCTEXT("OperatorStackTabName", "Operator Stack");
+
+	TabManager->RegisterTabSpawner(
+		UOperatorStackEditorSubsystem::TabId,
+		FOnSpawnTab::CreateSP(this, &FOperatorStackEditorTabInstance::OnSpawnTab)
+	)
+	.SetDisplayNameAttribute(DisplayNameText)
+	.SetDisplayName(DisplayNameText)
+	.SetGroup(WorkspaceMenu::GetMenuStructure().GetLevelEditorCategory())
+	.SetIcon(FSlateIcon(FAppStyle::GetAppStyleSetName(), "ClassIcon.UserDefinedStruct"));
+
+	// To auto-invoke tab
+	UnbindDelegates();
+	BindDelegates();
+
+	return true;
+}
+
+TSharedRef<SDockTab> FOperatorStackEditorTabInstance::OnSpawnTab(const FSpawnTabArgs& InArgs)
+{
+	const TSharedPtr<ILevelEditor> LevelEditor = LevelEditorWeak.Pin();
+	UOperatorStackEditorSubsystem* Subsystem = UOperatorStackEditorSubsystem::Get();
+
+	check(!!Subsystem)
+	check(LevelEditor.IsValid())
+
+	const TSharedRef<SOperatorStackEditorWidget> Widget = Subsystem->GenerateWidget();
+	WidgetIdentifier = Widget->GetPanelId();
+
+	OnSelectionSetChanged(LevelEditor->GetElementSelectionSet());
+
+	FTSTicker::GetCoreTicker().AddTicker(FTickerDelegate::CreateSPLambda(this, [this](float)
+	{
+		const UOperatorStackEditorSubsystem* Subsystem = UOperatorStackEditorSubsystem::Get();
+		const TSharedPtr<SOperatorStackEditorWidget> Widget = GetOperatorStackEditorWidget();
+		
+		if (Subsystem && Widget)
+		{
+			Subsystem->OnOperatorStackSpawnedDelegate.Broadcast(Widget.ToSharedRef());
+		}
+
+		return false;
+	}), 0);
+
+	return SNew(SDockTab)
+		.TabRole(ETabRole::PanelTab)
+		.Content()
+		[
+			Widget
+		];
+}
+
+void FOperatorStackEditorTabInstance::OnSelectionSetChanged(const UTypedElementSelectionSet* InSelection)
+{
+	if (!InSelection || !LevelEditorWeak.IsValid())
+	{
+		return;
+	}
+	
+	const TArray<UObject*> SelectedObjects = InSelection->GetSelectedObjects();
+
+	TSharedPtr<SOperatorStackEditorWidget> Widget = GetOperatorStackEditorWidget();
+
+	TArray<FOperatorStackEditorItemPtr> SelectedItems;
+	Algo::Transform(SelectedObjects, SelectedItems, [](UObject* InObject)
+	{
+		return MakeShared<FOperatorStackEditorObjectItem>(InObject);
+	});
+	const FOperatorStackEditorContext Context(SelectedItems);
+
+	if (!Widget.IsValid())
+	{
+		// Invoke tab
+		const UOperatorStackEditorSubsystem* Subsystem = UOperatorStackEditorSubsystem::Get();
+
+		check(!!Subsystem)
+
+		bool bInvokeTab = !Subsystem->ForEachCustomization([this, &Context](UOperatorStackEditorStackCustomization* InCustomization)->bool
+		{
+			if (InCustomization && InCustomization->ShouldFocusCustomization(Context))
+			{
+				return false;
+			}
+				
+			return true;
+		});
+
+		if (bInvokeTab)
+		{
+			InvokeTab();
+			Widget = GetOperatorStackEditorWidget();
+		}
+	}
+
+	if (Widget.IsValid())
+	{
+		Widget->SetContext(Context);
+	}
+}
+
+void FOperatorStackEditorTabInstance::OnSelectionChanged(UObject* InSelectionObject)
+{
+	if (const USelection* Selection = Cast<USelection>(InSelectionObject))
+	{
+		OnSelectionSetChanged(Selection->GetElementSelectionSet());
+	}
+}
+
+bool FOperatorStackEditorTabInstance::UnregisterTab()
+{
+	const TSharedPtr<ILevelEditor> LevelEditor = LevelEditorWeak.Pin();
+
+	if (!LevelEditor.IsValid())
+	{
+		return false;
+	}
+
+	const TSharedPtr<FTabManager> TabManager = LevelEditor->GetTabManager();
+
+	if (!TabManager.IsValid())
+	{
+		return false;
+	}
+
+	return TabManager->UnregisterTabSpawner(UOperatorStackEditorSubsystem::TabId);
+}
+
+void FOperatorStackEditorTabInstance::BindDelegates()
+{
+	const TSharedPtr<ILevelEditor> LevelEditor = LevelEditorWeak.Pin();
+
+	if (!LevelEditor.IsValid())
+	{
+		return;
+	}
+
+	LevelEditor->GetMutableElementSelectionSet()->OnChanged().AddSP(this, &FOperatorStackEditorTabInstance::OnSelectionSetChanged);
+	USelection::SelectionChangedEvent.AddSP(this, &FOperatorStackEditorTabInstance::OnSelectionChanged);
+}
+
+void FOperatorStackEditorTabInstance::UnbindDelegates()
+{
+	const TSharedPtr<ILevelEditor> LevelEditor = LevelEditorWeak.Pin();
+
+	if (!LevelEditor.IsValid())
+	{
+		return;
+	}
+
+	LevelEditor->GetMutableElementSelectionSet()->OnChanged().RemoveAll(this);
+	USelection::SelectionChangedEvent.RemoveAll(this);
+}
+
+#undef LOCTEXT_NAMESPACE
\ No newline at end of file
diff --git a/Engine/Plugins/Experimental/OperatorStack/Source/OperatorStack/Private/Widgets/Tabs/OperatorStackEditorTabInstance.h b/Engine/Plugins/Experimental/OperatorStack/Source/OperatorStack/Private/Widgets/Tabs/OperatorStackEditorTabInstance.h
new file mode 100644
--- /dev/null
+++ b/Engine/Plugins/Experimental/OperatorStack/Source/OperatorStack/Private/Widgets/Tabs/OperatorStackEditorTabInstance.h
@@ -0,0 +1,43 @@
+// Copyright Epic Games, Inc. All Rights Reserved.
+
+#pragma once
+
+#include "Templates/SharedPointer.h"
+
+class FSpawnTabArgs;
+class ILevelEditor;
+class SDockTab;
+class SOperatorStackEditorWidget;
+class UObject;
+class UTypedElementSelectionSet;
+
+struct FOperatorStackEditorTabInstance : TSharedFromThis<FOperatorStackEditorTabInstance>
+{
+	explicit FOperatorStackEditorTabInstance(const TSharedRef<ILevelEditor>& InLevelEditor);
+	virtual ~FOperatorStackEditorTabInstance();
+
+	TSharedPtr<SDockTab> InvokeTab();
+	bool CloseTab();
+	bool RefreshTab(UObject* InContext, bool bInForce);
+	bool FocusTab(const UObject* InContext, FName InIdentifier);
+
+	TSharedPtr<SOperatorStackEditorWidget> GetOperatorStackEditorWidget() const;
+	bool RegisterTab();
+	bool UnregisterTab();
+
+	TSharedPtr<ILevelEditor> GetLevelEditor() const
+	{
+		return LevelEditorWeak.Pin();
+	}
+
+private:
+	void BindDelegates();
+	void UnbindDelegates();
+
+	TSharedRef<SDockTab> OnSpawnTab(const FSpawnTabArgs& InArgs);
+	void OnSelectionSetChanged(const UTypedElementSelectionSet* InSelection);
+	void OnSelectionChanged(UObject* InSelectionObject);
+
+	TWeakPtr<ILevelEditor> LevelEditorWeak;
+	int32 WidgetIdentifier = INDEX_NONE;
+};
diff --git a/Engine/Plugins/Experimental/OperatorStack/Source/OperatorStack/Public/Customizations/OperatorStackEditorStackCustomization.h b/Engine/Plugins/Experimental/OperatorStack/Source/OperatorStack/Public/Customizations/OperatorStackEditorStackCustomization.h
--- a/Engine/Plugins/Experimental/OperatorStack/Source/OperatorStack/Public/Customizations/OperatorStackEditorStackCustomization.h
+++ b/Engine/Plugins/Experimental/OperatorStack/Source/OperatorStack/Public/Customizations/OperatorStackEditorStackCustomization.h
@@ -120,6 +120,18 @@ public:
 
 	/** Get the displayed icon of this customization */
 	OPERATORSTACKEDITOR_API virtual const FSlateBrush* GetIcon() const;
+	
+	/** Refresh the selection linked to this context */
+	OPERATORSTACKEDITOR_API void RefreshActiveSelection(UObject* InContext, bool bInForce) const;
+
+	/** Switches section of operator stack and select this customization for this context */
+	OPERATORSTACKEDITOR_API void FocusCustomization(const UObject* InContext) const;
+
+	/** Whether new context should focus on this customization in the widget */
+	virtual bool ShouldFocusCustomization(const FOperatorStackEditorContext& InContext) const
+	{
+		return false;
+	};
 
 private:
 	/** Unique identifier of this customization */
diff --git a/Engine/Plugins/Experimental/OperatorStack/Source/OperatorStack/Public/Subsystems/OperatorStackEditorSubsystem.h b/Engine/Plugins/Experimental/OperatorStack/Source/OperatorStack/Public/Subsystems/OperatorStackEditorSubsystem.h
--- a/Engine/Plugins/Experimental/OperatorStack/Source/OperatorStack/Public/Subsystems/OperatorStackEditorSubsystem.h
+++ b/Engine/Plugins/Experimental/OperatorStack/Source/OperatorStack/Public/Subsystems/OperatorStackEditorSubsystem.h
@@ -7,8 +7,11 @@
 #include "UObject/WeakObjectPtrTemplates.h"
 #include "OperatorStackEditorSubsystem.generated.h"
 
+class ILevelEditor;
 class SOperatorStackEditorWidget;
+class UObject;
 class UOperatorStackEditorStackCustomization;
+struct FOperatorStackEditorTabInstance;
 
 /** Subsystem that handles operator stack customization */
 UCLASS()
@@ -17,10 +20,19 @@ class UOperatorStackEditorSubsystem : public UEditorSubsystem
 	GENERATED_BODY()
 
 	friend class SOperatorStackEditorPanel;
+	friend struct FOperatorStackEditorTabInstance;
 
 public:
+	OPERATORSTACKEDITOR_API static const FName TabId;
 	OPERATORSTACKEDITOR_API static UOperatorStackEditorSubsystem* Get();
-
+	
+	DECLARE_MULTICAST_DELEGATE_OneParam(FOnOperatorStackSpawned, TSharedRef<SOperatorStackEditorWidget> /** InWidget */)
+
+	static FOnOperatorStackSpawned::RegistrationType& OnOperatorStackSpawned()
+	{
+		return OnOperatorStackSpawnedDelegate;
+	}
+	
 	//~ Begin USubsystem
 	virtual void Initialize(FSubsystemCollectionBase& Collection) override;
 	virtual void Deinitialize() override;
@@ -41,18 +53,34 @@ public:
 	/** Finds an existing operator stack widget */
 	TSharedPtr<SOperatorStackEditorWidget> FindWidget(int32 InId);
 
+	/** Finds existing operator stack widgets within a provided context */
+	TArray<TSharedPtr<SOperatorStackEditorWidget>> FindWidgets(const UWorld* InContext);
+
 	/** Loops through each operator stack widgets available */
 	OPERATORSTACKEDITOR_API bool ForEachCustomizationWidget(TFunctionRef<bool(TSharedRef<SOperatorStackEditorWidget>)> InFunction) const;
 
+	/** Refresh only selection of widget with specific context */
+	OPERATORSTACKEDITOR_API void RefreshCustomizationWidget(UObject* InContext, bool bInForce);
+	
+	/** Focus customization widget with specific context */
+	OPERATORSTACKEDITOR_API void FocusCustomizationWidget(const UObject* InContext, FName InIdentifier);
+
 protected:
+	OPERATORSTACKEDITOR_API static FOnOperatorStackSpawned OnOperatorStackSpawnedDelegate;
+
     void ScanForStackCustomizations();
 
 	void OnWidgetDestroyed(int32 InPanelId);
 
+	void OnLevelEditorCreated(TSharedPtr<ILevelEditor> InLevelEditor);
+
 	/** Map of identifier and stack customization models */
 	UPROPERTY()
 	TMap<FName, TObjectPtr<UOperatorStackEditorStackCustomization>> CustomizationStacks;
 
 	/** Current widget created associated to an identifier */
 	TMap<int32, TWeakPtr<SOperatorStackEditorWidget>> CustomizationWidgets;
+
+	/** Current tab within editor */
+	TArray<TSharedPtr<FOperatorStackEditorTabInstance>> TabInstances;
 };
diff --git a/Engine/Plugins/Experimental/OperatorStack/Source/OperatorStack/Public/Widgets/SOperatorStackEditorWidget.h b/Engine/Plugins/Experimental/OperatorStack/Source/OperatorStack/Public/Widgets/SOperatorStackEditorWidget.h
--- a/Engine/Plugins/Experimental/OperatorStack/Source/OperatorStack/Public/Widgets/SOperatorStackEditorWidget.h
+++ b/Engine/Plugins/Experimental/OperatorStack/Source/OperatorStack/Public/Widgets/SOperatorStackEditorWidget.h
@@ -7,6 +7,7 @@
 
 class FDetailColumnSizeData;
 class IDetailKeyframeHandler;
+class UOperatorStackEditorStackCustomization;
 struct FOperatorStackEditorContext;
 
 /** Main operator stack widget created, only exposed function */
@@ -17,7 +18,10 @@ public:
 	virtual void SetActiveCustomization(const FName& InCustomization) {}
 
 	/** Sets the customization to show in the toolbar, overrides order */
-	virtual void SetToolbarCustomization(const TArray<FName>& InCustomizations) {}
+	virtual void SetToolbarCustomizations(const TArray<FName>& InCustomizations) {}
+
+	/** Get all the customizations object active on this widget */
+	virtual void GetToolbarCustomizations(TArray<UOperatorStackEditorStackCustomization*>& OutCustomizations) {}
 
 	/** Show the customization in the panel */
 	virtual void ShowToolbarCustomization(FName InCustomization) {}
diff --git a/Engine/Plugins/Experimental/PropertyAnimatorCore/Source/PropertyAnimatorCore/Private/Components/PropertyAnimatorCoreComponent.cpp b/Engine/Plugins/Experimental/PropertyAnimatorCore/Source/PropertyAnimatorCore/Private/Components/PropertyAnimatorCoreComponent.cpp
--- a/Engine/Plugins/Experimental/PropertyAnimatorCore/Source/PropertyAnimatorCore/Private/Components/PropertyAnimatorCoreComponent.cpp
+++ b/Engine/Plugins/Experimental/PropertyAnimatorCore/Source/PropertyAnimatorCore/Private/Components/PropertyAnimatorCoreComponent.cpp
@@ -232,22 +232,7 @@ void UPropertyAnimatorCoreComponent::OnComponentCreated()
 {
 	Super::OnComponentCreated();
 
-	if (AActor* OwningActor = GetOwner())
-	{
-		// For spawnable templates, restore and resolve properties owner
-		constexpr bool bForceRestore = true;
-
-		for (const TObjectPtr<UPropertyAnimatorCoreBase>& Animator : PropertyAnimators)
-		{
-			if (Animator)
-			{
-				Animator->RestoreProperties(bForceRestore);
-				Animator->ResolvePropertiesOwner(OwningActor);
-			}
-		}
-	}
-
-	UPropertyAnimatorCoreBase::OnAnimatorAddedDelegate.Broadcast(this, nullptr);
+	InitAnimatorComponent();
 }
 
 UPropertyAnimatorCoreComponent* UPropertyAnimatorCoreComponent::FindOrAdd(AActor* InActor)
@@ -393,11 +378,23 @@ FName UPropertyAnimatorCoreComponent::GetPropertyAnimatorsPropertyName()
 	return GET_MEMBER_NAME_CHECKED(UPropertyAnimatorCoreComponent, PropertyAnimators);
 }
 
+void UPropertyAnimatorCoreComponent::PreEditUndo()
+{
+	Super::PreEditUndo();
+
+	PropertyAnimatorsInternal = PropertyAnimators;
+}
+
 void UPropertyAnimatorCoreComponent::PostEditUndo()
 {
 	Super::PostEditUndo();
 
-	OnAnimatorsEnabledChanged();
+	if (!bRegistered && IsValid(this))
+	{
+		UPropertyAnimatorCoreBase::OnAnimatorAddedDelegate.Broadcast(this, nullptr);
+	}
+
+	OnAnimatorsChanged();
 }
 
 void UPropertyAnimatorCoreComponent::PreEditChange(FProperty* PropertyAboutToChange)
@@ -539,6 +536,26 @@ bool UPropertyAnimatorCoreComponent::EvaluateAnimators()
 	return true;
 }
 
+void UPropertyAnimatorCoreComponent::InitAnimatorComponent()
+{
+	if (AActor* OwningActor = GetOwner())
+	{
+		// For spawnable templates, restore and resolve properties owner
+		constexpr bool bForceRestore = true;
+
+		for (const TObjectPtr<UPropertyAnimatorCoreBase>& Animator : PropertyAnimators)
+		{
+			if (Animator)
+			{
+				Animator->RestoreProperties(bForceRestore);
+				Animator->ResolvePropertiesOwner(OwningActor);
+			}
+		}
+	}
+
+	UPropertyAnimatorCoreBase::OnAnimatorAddedDelegate.Broadcast(this, nullptr);
+}
+
 TArray<FName> UPropertyAnimatorCoreComponent::GetTimeSourceNames() const
 {
 	TArray<FName> TimeSourceNames;
diff --git a/Engine/Plugins/Experimental/PropertyAnimatorCore/Source/PropertyAnimatorCore/Public/Components/PropertyAnimatorCoreComponent.h b/Engine/Plugins/Experimental/PropertyAnimatorCore/Source/PropertyAnimatorCore/Public/Components/PropertyAnimatorCoreComponent.h
--- a/Engine/Plugins/Experimental/PropertyAnimatorCore/Source/PropertyAnimatorCore/Public/Components/PropertyAnimatorCoreComponent.h
+++ b/Engine/Plugins/Experimental/PropertyAnimatorCore/Source/PropertyAnimatorCore/Public/Components/PropertyAnimatorCoreComponent.h
@@ -81,6 +81,7 @@ protected:
 	virtual void PostEditImport() override;
 	virtual void PostDuplicate(EDuplicateMode::Type InMode) override;
 #if WITH_EDITOR
+	virtual void PreEditUndo() override;
 	virtual void PostEditUndo() override;
 	virtual void PreEditChange(FProperty* PropertyAboutToChange) override;
 	virtual void PostEditChangeProperty(FPropertyChangedEvent& PropertyChangedEvent) override;
@@ -119,6 +120,8 @@ protected:
 	/** Evaluate only specified animators */
 	bool EvaluateAnimators();
 
+	void InitAnimatorComponent();
+
 	/** Finds a cached time source with this name or creates a new one */
 	UPropertyAnimatorCoreTimeSourceBase* FindOrAddTimeSource(FName InTimeSourceName);
 
@@ -152,7 +155,7 @@ private:
 	TSet<TObjectPtr<UPropertyAnimatorCoreBase>> Animators;
 
 	/** Transient copy of property animators when changes are detected to see the diff only */
-	UPROPERTY(Transient, DuplicateTransient, TextExportTransient)
+	UPROPERTY(Transient, DuplicateTransient, TextExportTransient, NonTransactional)
 	TArray<TObjectPtr<UPropertyAnimatorCoreBase>> PropertyAnimatorsInternal;
 
 	/** Cached time sources used by this animator component */
diff --git a/Engine/Plugins/Experimental/PropertyAnimatorCore/Source/PropertyAnimatorCoreEditor/Private/Customizations/PropertyAnimatorCoreEditorStackCustomization.cpp b/Engine/Plugins/Experimental/PropertyAnimatorCore/Source/PropertyAnimatorCoreEditor/Private/Customizations/PropertyAnimatorCoreEditorStackCustomization.cpp
--- a/Engine/Plugins/Experimental/PropertyAnimatorCore/Source/PropertyAnimatorCoreEditor/Private/Customizations/PropertyAnimatorCoreEditorStackCustomization.cpp
+++ b/Engine/Plugins/Experimental/PropertyAnimatorCore/Source/PropertyAnimatorCoreEditor/Private/Customizations/PropertyAnimatorCoreEditorStackCustomization.cpp
@@ -30,6 +30,18 @@ UPropertyAnimatorCoreEditorStackCustomization::UPropertyAnimatorCoreEditorStackC
 {
 	RegisterCustomizationFor(UPropertyAnimatorCoreBase::StaticClass());
 	RegisterCustomizationFor(UPropertyAnimatorCoreComponent::StaticClass());
+
+	// Animator delegates
+	UPropertyAnimatorCoreBase::OnPropertyAnimatorAdded().AddUObject(this, &UPropertyAnimatorCoreEditorStackCustomization::OnAnimatorUpdated);
+	UPropertyAnimatorCoreBase::OnPropertyAnimatorRemoved().AddUObject(this, &UPropertyAnimatorCoreEditorStackCustomization::OnAnimatorRemoved);
+	UPropertyAnimatorCoreBase::OnPropertyAnimatorRenamed().AddUObject(this, &UPropertyAnimatorCoreEditorStackCustomization::OnAnimatorUpdated);
+}
+
+UPropertyAnimatorCoreEditorStackCustomization::~UPropertyAnimatorCoreEditorStackCustomization()
+{
+	UPropertyAnimatorCoreBase::OnPropertyAnimatorAdded().RemoveAll(this);
+	UPropertyAnimatorCoreBase::OnPropertyAnimatorRemoved().RemoveAll(this);
+	UPropertyAnimatorCoreBase::OnPropertyAnimatorRenamed().RemoveAll(this);
 }
 
 bool UPropertyAnimatorCoreEditorStackCustomization::GetRootItem(const FOperatorStackEditorContext& InContext, FOperatorStackEditorItemPtr& OutRootItem) const
@@ -325,6 +337,17 @@ const FSlateBrush* UPropertyAnimatorCoreEditorStackCustomization::GetIcon() cons
 	return FSlateIconFinder::FindIconForClass(UPropertyAnimatorCoreBase::StaticClass()).GetIcon();
 }
 
+bool UPropertyAnimatorCoreEditorStackCustomization::ShouldFocusCustomization(const FOperatorStackEditorContext& InContext) const
+{
+	if (InContext.GetItems().IsEmpty())
+	{
+		return false;
+	}
+	
+	return InContext.GetItems().Last().IsValid()
+		&& (InContext.GetItems().Last()->IsA<UPropertyAnimatorCoreComponent>() || InContext.GetItems().Last()->IsA<UPropertyAnimatorCoreBase>());
+}
+
 void UPropertyAnimatorCoreEditorStackCustomization::RemoveAnimatorAction(FOperatorStackEditorItemPtr InItem) const
 {
 	if (!InItem.IsValid() || !InItem->HasValue())
@@ -504,6 +527,22 @@ void UPropertyAnimatorCoreEditorStackCustomization::FillAnimatorContextActionMen
 	InToolMenu->AddMenuEntry(RemoveAnimatorMenuEntry.Name, RemoveAnimatorMenuEntry);
 }
 
+void UPropertyAnimatorCoreEditorStackCustomization::OnAnimatorUpdated(UPropertyAnimatorCoreComponent* InComponent, UPropertyAnimatorCoreBase* InAnimator)
+{
+	if (InComponent)
+	{
+		RefreshActiveSelection(InComponent, /** Force */false);
+	}
+}
+
+void UPropertyAnimatorCoreEditorStackCustomization::OnAnimatorRemoved(UPropertyAnimatorCoreComponent* InComponent, UPropertyAnimatorCoreBase* InAnimator)
+{
+	if (InComponent)
+	{
+		RefreshActiveSelection(InComponent, /** Force */true);
+	}
+}
+
 TSharedRef<FUICommandList> UPropertyAnimatorCoreEditorStackCustomization::CreateAnimatorCommands(FOperatorStackEditorItemPtr InItem)
 {
 	TSharedRef<FUICommandList> Commands = MakeShared<FUICommandList>();
diff --git a/Engine/Plugins/Experimental/PropertyAnimatorCore/Source/PropertyAnimatorCoreEditor/Private/Customizations/PropertyAnimatorCoreEditorStackCustomization.h b/Engine/Plugins/Experimental/PropertyAnimatorCore/Source/PropertyAnimatorCoreEditor/Private/Customizations/PropertyAnimatorCoreEditorStackCustomization.h
--- a/Engine/Plugins/Experimental/PropertyAnimatorCore/Source/PropertyAnimatorCoreEditor/Private/Customizations/PropertyAnimatorCoreEditorStackCustomization.h
+++ b/Engine/Plugins/Experimental/PropertyAnimatorCore/Source/PropertyAnimatorCoreEditor/Private/Customizations/PropertyAnimatorCoreEditorStackCustomization.h
@@ -17,6 +17,7 @@ class UPropertyAnimatorCoreEditorStackCustomization : public UOperatorStackEdito
 
 public:
 	UPropertyAnimatorCoreEditorStackCustomization();
+	virtual ~UPropertyAnimatorCoreEditorStackCustomization() override;
 
 	//~ Begin UOperatorStackEditorStackCustomization
 	virtual bool GetRootItem(const FOperatorStackEditorContext& InContext, FOperatorStackEditorItemPtr& OutRootItem) const override;
@@ -26,6 +27,7 @@ public:
 	virtual void CustomizeItemBody(const FOperatorStackEditorItemPtr& InItem, const FOperatorStackEditorTree& InItemTree, FOperatorStackEditorBodyBuilder& InBodyBuilder) override;
 	virtual bool OnIsItemSelectable(const FOperatorStackEditorItemPtr& InItem) override;
 	virtual const FSlateBrush* GetIcon() const override;
+	virtual bool ShouldFocusCustomization(const FOperatorStackEditorContext& InContext) const override;
 	//~ End UOperatorStackEditorStackCustomization
 
 protected:
@@ -40,6 +42,9 @@ protected:
 	void FillAnimatorHeaderActionMenu(UToolMenu* InToolMenu);
 	void FillAnimatorContextActionMenu(UToolMenu* InToolMenu) const;
 
+	void OnAnimatorUpdated(UPropertyAnimatorCoreComponent* InComponent, UPropertyAnimatorCoreBase* InAnimator);
+	void OnAnimatorRemoved(UPropertyAnimatorCoreComponent* InComponent, UPropertyAnimatorCoreBase* InAnimator);
+
 	/** Create commands for animator */
 	TSharedRef<FUICommandList> CreateAnimatorCommands(FOperatorStackEditorItemPtr InItem);
 };
\ No newline at end of file
