MotionDesign : OperatorStack
- Fixed issue happening when selecting actor and its components (modifier/animator), stack items would show as multiple/different items when they are the same

diff --git a/Engine/Plugins/Experimental/OperatorStack/Source/OperatorStack/Public/Items/OperatorStackEditorGroupItem.h b/Engine/Plugins/Experimental/OperatorStack/Source/OperatorStack/Public/Items/OperatorStackEditorGroupItem.h
--- a/Engine/Plugins/Experimental/OperatorStack/Source/OperatorStack/Public/Items/OperatorStackEditorGroupItem.h
+++ b/Engine/Plugins/Experimental/OperatorStack/Source/OperatorStack/Public/Items/OperatorStackEditorGroupItem.h
@@ -14,6 +14,8 @@ struct FOperatorStackEditorGroupItem : FOperatorStackEditorItem
 	{
 		CachedHash = 0;
 
+		TSet<uint32> Hashes;
+		Hashes.Reserve(InItems.Num());
 		for (const FOperatorStackEditorItemPtr& Item : InItems)
 		{
 			if (Item.IsValid())
@@ -24,8 +26,13 @@ struct FOperatorStackEditorGroupItem : FOperatorStackEditorItem
 				// Group is only allowed for item of the same type
 				check(GetValueType() == Item->GetValueType())
 
-				CachedHash += Item->GetHash();
-				Items.Add(Item);
+				const uint32 ItemHash = Item->GetHash();
+				if (!Hashes.Contains(ItemHash))
+				{
+					CachedHash += ItemHash;
+					Items.Add(Item);
+					Hashes.Add(ItemHash);
+				}
 			}
 		}
 	}
