MotionDesign : OperatorStack
- Fixed selecting an asset from Content browser will update operator stack view

diff --git a/Engine/Plugins/Experimental/OperatorStack/Source/OperatorStack/Private/Widgets/SOperatorStackEditorPanel.cpp b/Engine/Plugins/Experimental/OperatorStack/Source/OperatorStack/Private/Widgets/SOperatorStackEditorPanel.cpp
--- a/Engine/Plugins/Experimental/OperatorStack/Source/OperatorStack/Private/Widgets/SOperatorStackEditorPanel.cpp
+++ b/Engine/Plugins/Experimental/OperatorStack/Source/OperatorStack/Private/Widgets/SOperatorStackEditorPanel.cpp
@@ -79,6 +79,11 @@ void SOperatorStackEditorPanel::Construct(const FArguments& InArgs)
 
 void SOperatorStackEditorPanel::SetContext(const FOperatorStackEditorContext& InContext)
 {
+	if (Context.IsValid() && (*Context.Get()) == InContext)
+	{
+		return;
+	}
+
 	Context = MakeShared<FOperatorStackEditorContext>(InContext);
 	RefreshContext();
 }
diff --git a/Engine/Plugins/Experimental/OperatorStack/Source/OperatorStack/Private/Widgets/Tabs/OperatorStackEditorTabInstance.cpp b/Engine/Plugins/Experimental/OperatorStack/Source/OperatorStack/Private/Widgets/Tabs/OperatorStackEditorTabInstance.cpp
--- a/Engine/Plugins/Experimental/OperatorStack/Source/OperatorStack/Private/Widgets/Tabs/OperatorStackEditorTabInstance.cpp
+++ b/Engine/Plugins/Experimental/OperatorStack/Source/OperatorStack/Private/Widgets/Tabs/OperatorStackEditorTabInstance.cpp
@@ -4,6 +4,7 @@
 
 #include "Contexts/OperatorStackEditorContext.h"
 #include "Customizations/OperatorStackEditorStackCustomization.h"
+#include "Editor.h"
 #include "Elements/Framework/TypedElementSelectionSet.h"
 #include "GameFramework/Actor.h"
 #include "ILevelEditor.h"
@@ -255,17 +256,59 @@ void FOperatorStackEditorTabInstance::OnSelectionSetChanged(const UTypedElementS
 	{
 		return;
 	}
-	
-	const TArray<UObject*> SelectedObjects = InSelection->GetSelectedObjects();
+
+	TSet<UObject*> SelectedObjects(InSelection->GetSelectedObjects());
+
+	if (USelection* ActorSelection = GEditor->GetSelectedActors())
+	{
+		TArray<UObject*> Actors;
+		ActorSelection->GetSelectedObjects(Actors);
+		SelectedObjects.Append(Actors);
+	}
+
+	if (USelection* ComponentSelection = GEditor->GetSelectedComponents())
+	{
+		TArray<UObject*> Components;
+		ComponentSelection->GetSelectedObjects(Components);
+		SelectedObjects.Append(Components);
+	}
+
+	if (USelection* ObjectSelection = GEditor->GetSelectedObjects())
+	{
+		TArray<UObject*> Objects;
+		ObjectSelection->GetSelectedObjects(Objects);
+		SelectedObjects.Append(Objects);
+	}
 
 	TSharedPtr<SOperatorStackEditorWidget> Widget = GetOperatorStackEditorWidget();
 
+	TSharedPtr<ILevelEditor> LevelEditor = LevelEditorWeak.Pin();
 	TArray<FOperatorStackEditorItemPtr> SelectedItems;
-	Algo::Transform(SelectedObjects, SelectedItems, [](UObject* InObject)
+	Algo::TransformIf(
+		SelectedObjects,
+		SelectedItems,
+		[LevelEditor](const UObject* InObject)
+		{
+			return InObject && InObject->GetWorld() == LevelEditor->GetWorld();
+		},
+		[](UObject* InObject)
+		{
+			return MakeShared<FOperatorStackEditorObjectItem>(InObject);
+		});
+
+	if (!Widget.IsValid() && SelectedItems.IsEmpty())
 	{
-		return MakeShared<FOperatorStackEditorObjectItem>(InObject);
-	});
-	const FOperatorStackEditorContext Context(SelectedItems);
+		return;
+	}
+
+	const FOperatorStackEditorContext NewContext(SelectedItems);
+
+	if (Widget.IsValid()
+		&& Widget->GetContext()
+		&& *Widget->GetContext().Get() == NewContext)
+	{
+		return;
+	}
 
 	if (!Widget.IsValid())
 	{
@@ -274,9 +317,9 @@ void FOperatorStackEditorTabInstance::OnSelectionSetChanged(const UTypedElementS
 
 		check(!!Subsystem)
 
-		bool bInvokeTab = !Subsystem->ForEachCustomization([this, &Context](UOperatorStackEditorStackCustomization* InCustomization)->bool
+		bool bInvokeTab = !Subsystem->ForEachCustomization([this, &NewContext](UOperatorStackEditorStackCustomization* InCustomization)->bool
 		{
-			if (InCustomization && InCustomization->ShouldFocusCustomization(Context))
+			if (InCustomization && InCustomization->ShouldFocusCustomization(NewContext))
 			{
 				return false;
 			}
@@ -293,7 +336,7 @@ void FOperatorStackEditorTabInstance::OnSelectionSetChanged(const UTypedElementS
 
 	if (Widget.IsValid())
 	{
-		Widget->SetContext(Context);
+		Widget->SetContext(NewContext);
 	}
 }
 
diff --git a/Engine/Plugins/Experimental/OperatorStack/Source/OperatorStack/Public/Contexts/OperatorStackEditorContext.h b/Engine/Plugins/Experimental/OperatorStack/Source/OperatorStack/Public/Contexts/OperatorStackEditorContext.h
--- a/Engine/Plugins/Experimental/OperatorStack/Source/OperatorStack/Public/Contexts/OperatorStackEditorContext.h
+++ b/Engine/Plugins/Experimental/OperatorStack/Source/OperatorStack/Public/Contexts/OperatorStackEditorContext.h
@@ -12,6 +12,24 @@ struct FOperatorStackEditorContext final
 	explicit FOperatorStackEditorContext(const TArray<TSharedPtr<FOperatorStackEditorItem>>& InItems)
 		: Items(InItems)
 	{
+		CachedHash = 0;
+
+		TSet<uint32> Hashes;
+		Hashes.Reserve(InItems.Num());
+		for (const TSharedPtr<FOperatorStackEditorItem>& Item : InItems)
+		{
+			if (Item.IsValid())
+			{
+				const uint32 ItemHash = Item->GetHash();
+				bool bIsAlreadyInSet = false;
+				Hashes.Add(ItemHash, &bIsAlreadyInSet);
+
+				if (!bIsAlreadyInSet)
+				{
+					CachedHash += ItemHash;
+				}
+			}
+		}
 	}
 
 	/** Items we want to customize */
@@ -19,10 +37,23 @@ struct FOperatorStackEditorContext final
 	{
 		return Items;
 	}
+	
+	friend uint32 GetTypeHash(const FOperatorStackEditorContext& InItem)
+	{
+		return InItem.CachedHash;
+	}
+
+	bool operator==(const FOperatorStackEditorContext& InOther) const
+	{
+		return CachedHash == InOther.CachedHash;
+	}
 
 protected:
 	/** Context items being customized */
 	TArray<TSharedPtr<FOperatorStackEditorItem>> Items;
+
+	/** Context items hash to compare contexts */
+	uint32 CachedHash = 0;
 };
 
 typedef TSharedPtr<FOperatorStackEditorContext> FOperatorStackEditorContextPtr;
