MotionDesign : OperatorStack
- Fixed selecting an asset from Content browser will update operator stack view

diff --git a/Engine/Plugins/Experimental/OperatorStack/Source/OperatorStack/Private/Widgets/SOperatorStackEditorPanel.cpp b/Engine/Plugins/Experimental/OperatorStack/Source/OperatorStack/Private/Widgets/SOperatorStackEditorPanel.cpp
--- a/Engine/Plugins/Experimental/OperatorStack/Source/OperatorStack/Private/Widgets/SOperatorStackEditorPanel.cpp
+++ b/Engine/Plugins/Experimental/OperatorStack/Source/OperatorStack/Private/Widgets/SOperatorStackEditorPanel.cpp
@@ -79,11 +79,6 @@ void SOperatorStackEditorPanel::Construct(const FArguments& InArgs)
 
 void SOperatorStackEditorPanel::SetContext(const FOperatorStackEditorContext& InContext)
 {
-	if (Context.IsValid() && (*Context.Get()) == InContext)
-	{
-		return;
-	}
-
 	Context = MakeShared<FOperatorStackEditorContext>(InContext);
 	RefreshContext();
 }
diff --git a/Engine/Plugins/Experimental/OperatorStack/Source/OperatorStack/Private/Widgets/Tabs/OperatorStackEditorTabInstance.cpp b/Engine/Plugins/Experimental/OperatorStack/Source/OperatorStack/Private/Widgets/Tabs/OperatorStackEditorTabInstance.cpp
--- a/Engine/Plugins/Experimental/OperatorStack/Source/OperatorStack/Private/Widgets/Tabs/OperatorStackEditorTabInstance.cpp
+++ b/Engine/Plugins/Experimental/OperatorStack/Source/OperatorStack/Private/Widgets/Tabs/OperatorStackEditorTabInstance.cpp
@@ -105,7 +105,7 @@ bool FOperatorStackEditorTabInstance::RefreshTab(UObject* InContext, bool bInFor
 
 	if (bInForce)
 	{
-		OnSelectionSetChanged(SelectionSet);
+		OnSelectionSetChanged(SelectionSet, bInForce);
 	}
 	else
 	{
@@ -227,7 +227,8 @@ TSharedRef<SDockTab> FOperatorStackEditorTabInstance::OnSpawnTab(const FSpawnTab
 	const TSharedRef<SOperatorStackEditorWidget> Widget = Subsystem->GenerateWidget();
 	WidgetIdentifier = Widget->GetPanelId();
 
-	OnSelectionSetChanged(LevelEditor->GetElementSelectionSet());
+	constexpr bool bForce = false;
+	OnSelectionSetChanged(LevelEditor->GetElementSelectionSet(), bForce);
 
 	FTSTicker::GetCoreTicker().AddTicker(FTickerDelegate::CreateSPLambda(this, [this](float)
 	{
@@ -250,7 +251,7 @@ TSharedRef<SDockTab> FOperatorStackEditorTabInstance::OnSpawnTab(const FSpawnTab
 		];
 }
 
-void FOperatorStackEditorTabInstance::OnSelectionSetChanged(const UTypedElementSelectionSet* InSelection)
+void FOperatorStackEditorTabInstance::OnSelectionSetChanged(const UTypedElementSelectionSet* InSelection, bool bInForce)
 {
 	if (!InSelection || !LevelEditorWeak.IsValid())
 	{
@@ -303,7 +304,8 @@ void FOperatorStackEditorTabInstance::OnSelectionSetChanged(const UTypedElementS
 
 	const FOperatorStackEditorContext NewContext(SelectedItems);
 
-	if (Widget.IsValid()
+	if (!bInForce
+		&& Widget.IsValid()
 		&& Widget->GetContext()
 		&& *Widget->GetContext().Get() == NewContext)
 	{
@@ -344,7 +346,8 @@ void FOperatorStackEditorTabInstance::OnSelectionChanged(UObject* InSelectionObj
 {
 	if (const USelection* Selection = Cast<USelection>(InSelectionObject))
 	{
-		OnSelectionSetChanged(Selection->GetElementSelectionSet());
+		constexpr bool bForce = false;
+		OnSelectionSetChanged(Selection->GetElementSelectionSet(), bForce);
 	}
 }
 
@@ -376,7 +379,8 @@ void FOperatorStackEditorTabInstance::BindDelegates()
 		return;
 	}
 
-	LevelEditor->GetMutableElementSelectionSet()->OnChanged().AddSP(this, &FOperatorStackEditorTabInstance::OnSelectionSetChanged);
+	constexpr bool bForce = false;
+	LevelEditor->GetMutableElementSelectionSet()->OnChanged().AddSP(this, &FOperatorStackEditorTabInstance::OnSelectionSetChanged, bForce);
 	USelection::SelectionChangedEvent.AddSP(this, &FOperatorStackEditorTabInstance::OnSelectionChanged);
 }
 
diff --git a/Engine/Plugins/Experimental/OperatorStack/Source/OperatorStack/Private/Widgets/Tabs/OperatorStackEditorTabInstance.h b/Engine/Plugins/Experimental/OperatorStack/Source/OperatorStack/Private/Widgets/Tabs/OperatorStackEditorTabInstance.h
--- a/Engine/Plugins/Experimental/OperatorStack/Source/OperatorStack/Private/Widgets/Tabs/OperatorStackEditorTabInstance.h
+++ b/Engine/Plugins/Experimental/OperatorStack/Source/OperatorStack/Private/Widgets/Tabs/OperatorStackEditorTabInstance.h
@@ -35,7 +35,7 @@ private:
 	void UnbindDelegates();
 
 	TSharedRef<SDockTab> OnSpawnTab(const FSpawnTabArgs& InArgs);
-	void OnSelectionSetChanged(const UTypedElementSelectionSet* InSelection);
+	void OnSelectionSetChanged(const UTypedElementSelectionSet* InSelection, bool bInForce);
 	void OnSelectionChanged(UObject* InSelectionObject);
 
 	TWeakPtr<ILevelEditor> LevelEditorWeak;
