MotionDesign : OperatorStack
- Fixed multiple selection change event causes tab to be invoked when it should not be

diff --git a/Engine/Plugins/VirtualProduction/OperatorStack/Source/OperatorStack/Private/Widgets/Tabs/OperatorStackEditorTabInstance.cpp b/Engine/Plugins/VirtualProduction/OperatorStack/Source/OperatorStack/Private/Widgets/Tabs/OperatorStackEditorTabInstance.cpp
--- a/Engine/Plugins/VirtualProduction/OperatorStack/Source/OperatorStack/Private/Widgets/Tabs/OperatorStackEditorTabInstance.cpp
+++ b/Engine/Plugins/VirtualProduction/OperatorStack/Source/OperatorStack/Private/Widgets/Tabs/OperatorStackEditorTabInstance.cpp
@@ -258,7 +258,37 @@ void FOperatorStackEditorTabInstance::OnSelectionSetChanged(const UTypedElementS
 		return;
 	}
 
-	TSet<UObject*> SelectedObjects(InSelection->GetSelectedObjects());
+	// Cancel previous request
+	if (SelectionDelegateHandle.IsValid())
+	{
+		FTSTicker::GetCoreTicker().RemoveTicker(SelectionDelegateHandle);
+		SelectionDelegateHandle.Reset();
+	}
+
+	// Wait one tick before processing since selection event can do : Select New -> Deselect Previous to avoid handling Previous again
+	TWeakObjectPtr<const UTypedElementSelectionSet> SelectionWeak(InSelection);
+	SelectionDelegateHandle = FTSTicker::GetCoreTicker().AddTicker(FTickerDelegate::CreateSP(this, &FOperatorStackEditorTabInstance::CheckActiveSelection, SelectionWeak, bInForce));
+}
+
+void FOperatorStackEditorTabInstance::OnSelectionChanged(UObject* InSelectionObject)
+{
+	if (const USelection* Selection = Cast<USelection>(InSelectionObject))
+	{
+		constexpr bool bForce = false;
+		OnSelectionSetChanged(Selection->GetElementSelectionSet(), bForce);
+	}
+}
+
+bool FOperatorStackEditorTabInstance::CheckActiveSelection(float, TWeakObjectPtr<const UTypedElementSelectionSet> InSelectionWeak, bool bInForce)
+{
+	const UTypedElementSelectionSet* Selection = InSelectionWeak.Get();
+
+	if (!GEditor || !Selection || IsEngineExitRequested())
+	{
+		return false; // Stop
+	}
+
+	TSet<UObject*> SelectedObjects(Selection->GetSelectedObjects());
 
 	if (USelection* ActorSelection = GEditor->GetSelectedActors())
 	{
@@ -299,7 +329,7 @@ void FOperatorStackEditorTabInstance::OnSelectionSetChanged(const UTypedElementS
 
 	if (!Widget.IsValid() && SelectedItems.IsEmpty())
 	{
-		return;
+		return false; // Stop
 	}
 
 	const FOperatorStackEditorContext NewContext(SelectedItems);
@@ -309,7 +339,7 @@ void FOperatorStackEditorTabInstance::OnSelectionSetChanged(const UTypedElementS
 		&& Widget->GetContext()
 		&& *Widget->GetContext().Get() == NewContext)
 	{
-		return;
+		return false; // Stop
 	}
 
 	if (!Widget.IsValid())
@@ -319,13 +349,15 @@ void FOperatorStackEditorTabInstance::OnSelectionSetChanged(const UTypedElementS
 
 		check(!!Subsystem)
 
-		bool bInvokeTab = !Subsystem->ForEachCustomization([this, &NewContext](UOperatorStackEditorStackCustomization* InCustomization)->bool
+		bool bInvokeTab = false;
+		Subsystem->ForEachCustomization([this, &NewContext, &bInvokeTab](UOperatorStackEditorStackCustomization* InCustomization)->bool
 		{
 			if (InCustomization && InCustomization->ShouldFocusCustomization(NewContext))
 			{
+				bInvokeTab = true;
 				return false;
 			}
-				
+
 			return true;
 		});
 
@@ -340,15 +372,8 @@ void FOperatorStackEditorTabInstance::OnSelectionSetChanged(const UTypedElementS
 	{
 		Widget->SetContext(NewContext);
 	}
-}
 
-void FOperatorStackEditorTabInstance::OnSelectionChanged(UObject* InSelectionObject)
-{
-	if (const USelection* Selection = Cast<USelection>(InSelectionObject))
-	{
-		constexpr bool bForce = false;
-		OnSelectionSetChanged(Selection->GetElementSelectionSet(), bForce);
-	}
+	return false; // Stop
 }
 
 bool FOperatorStackEditorTabInstance::UnregisterTab()
diff --git a/Engine/Plugins/VirtualProduction/OperatorStack/Source/OperatorStack/Private/Widgets/Tabs/OperatorStackEditorTabInstance.h b/Engine/Plugins/VirtualProduction/OperatorStack/Source/OperatorStack/Private/Widgets/Tabs/OperatorStackEditorTabInstance.h
--- a/Engine/Plugins/VirtualProduction/OperatorStack/Source/OperatorStack/Private/Widgets/Tabs/OperatorStackEditorTabInstance.h
+++ b/Engine/Plugins/VirtualProduction/OperatorStack/Source/OperatorStack/Private/Widgets/Tabs/OperatorStackEditorTabInstance.h
@@ -2,7 +2,9 @@
 
 #pragma once
 
+#include "Containers/Ticker.h"
 #include "Templates/SharedPointer.h"
+#include "UObject/WeakObjectPtr.h"
 
 class FSpawnTabArgs;
 class ILevelEditor;
@@ -38,6 +40,9 @@ private:
 	void OnSelectionSetChanged(const UTypedElementSelectionSet* InSelection, bool bInForce);
 	void OnSelectionChanged(UObject* InSelectionObject);
 
+	bool CheckActiveSelection(float, TWeakObjectPtr<const UTypedElementSelectionSet> InSelectionWeak, bool bInForce);
+
 	TWeakPtr<ILevelEditor> LevelEditorWeak;
 	int32 WidgetIdentifier = INDEX_NONE;
+	FTSTicker::FDelegateHandle SelectionDelegateHandle;
 };
